Rename ZenWidgets module into StorageServerWidgets and adjust referencing modules to use the new name.

diff --git a/Engine/Source/Editor/StorageServerWidgets/Internal/SZenCacheStatistics.h b/Engine/Source/Editor/StorageServerWidgets/Internal/SZenCacheStatistics.h
new file mode 100644
--- /dev/null
+++ b/Engine/Source/Editor/StorageServerWidgets/Internal/SZenCacheStatistics.h
@@ -0,0 +1,33 @@
+// Copyright Epic Games, Inc. All Rights Reserved.
+
+#pragma once
+
+#include "Experimental/ZenServerInterface.h"
+#include "ZenServiceInstanceManager.h"
+#include "Templates/SharedPointer.h"
+#include "Widgets/DeclarativeSyntaxSupport.h"
+#include "Widgets/SBoxPanel.h"
+#include "Widgets/SCompoundWidget.h"
+#include "Widgets/SWidget.h"
+
+class STORAGESERVERWIDGETS_API SZenCacheStatistics : public SCompoundWidget
+{
+	SLATE_BEGIN_ARGS(SZenCacheStatistics)
+		: _ZenServiceInstance(nullptr)
+	{ }
+
+	SLATE_ATTRIBUTE(TSharedPtr<UE::Zen::FZenServiceInstance>, ZenServiceInstance);
+
+	SLATE_END_ARGS()
+
+	void Construct(const FArguments& InArgs);
+
+private:
+
+	TSharedRef<SWidget> GetGridPanel();
+
+	EActiveTimerReturnType UpdateGridPanels(double InCurrentTime, float InDeltaTime);
+
+	SVerticalBox::FSlot* GridSlot = nullptr;
+	TAttribute<TSharedPtr<UE::Zen::FZenServiceInstance>> ZenServiceInstance;
+};
\ No newline at end of file
diff --git a/Engine/Source/Editor/StorageServerWidgets/Internal/SZenCidStoreStatistics.h b/Engine/Source/Editor/StorageServerWidgets/Internal/SZenCidStoreStatistics.h
new file mode 100644
--- /dev/null
+++ b/Engine/Source/Editor/StorageServerWidgets/Internal/SZenCidStoreStatistics.h
@@ -0,0 +1,33 @@
+// Copyright Epic Games, Inc. All Rights Reserved.
+
+#pragma once
+
+#include "Experimental/ZenServerInterface.h"
+#include "ZenServiceInstanceManager.h"
+#include "Templates/SharedPointer.h"
+#include "Widgets/DeclarativeSyntaxSupport.h"
+#include "Widgets/SBoxPanel.h"
+#include "Widgets/SCompoundWidget.h"
+#include "Widgets/SWidget.h"
+
+class STORAGESERVERWIDGETS_API SZenCidStoreStatistics : public SCompoundWidget
+{
+	SLATE_BEGIN_ARGS(SZenCidStoreStatistics)
+		: _ZenServiceInstance(nullptr)
+	{ }
+
+	SLATE_ATTRIBUTE(TSharedPtr<UE::Zen::FZenServiceInstance>, ZenServiceInstance);
+
+	SLATE_END_ARGS()
+
+	void Construct(const FArguments& InArgs);
+
+private:
+
+	TSharedRef<SWidget> GetGridPanel();
+
+	EActiveTimerReturnType UpdateGridPanels(double InCurrentTime, float InDeltaTime);
+
+	SVerticalBox::FSlot* GridSlot = nullptr;
+	TAttribute<TSharedPtr<UE::Zen::FZenServiceInstance>> ZenServiceInstance;
+};
\ No newline at end of file
diff --git a/Engine/Source/Editor/StorageServerWidgets/Internal/SZenDialogs.h b/Engine/Source/Editor/StorageServerWidgets/Internal/SZenDialogs.h
new file mode 100644
--- /dev/null
+++ b/Engine/Source/Editor/StorageServerWidgets/Internal/SZenDialogs.h
@@ -0,0 +1,21 @@
+// Copyright Epic Games, Inc. All Rights Reserved.
+
+#pragma once
+
+#include "Templates/SharedPointer.h"
+#include "Types/SlateEnums.h"
+#include "Widgets/DeclarativeSyntaxSupport.h"
+#include "Widgets/SBoxPanel.h"
+#include "Widgets/SCompoundWidget.h"
+
+class SWidget;
+
+class STORAGESERVERWIDGETS_API SZenStoreStausDialog : public SCompoundWidget
+{
+	SLATE_BEGIN_ARGS(SZenStoreStausDialog) {}
+	SLATE_END_ARGS()
+
+	void Construct(const FArguments& InArgs);
+
+private:
+};
diff --git a/Engine/Source/Editor/StorageServerWidgets/Internal/SZenProjectStatistics.h b/Engine/Source/Editor/StorageServerWidgets/Internal/SZenProjectStatistics.h
new file mode 100644
--- /dev/null
+++ b/Engine/Source/Editor/StorageServerWidgets/Internal/SZenProjectStatistics.h
@@ -0,0 +1,33 @@
+// Copyright Epic Games, Inc. All Rights Reserved.
+
+#pragma once
+
+#include "Experimental/ZenServerInterface.h"
+#include "ZenServiceInstanceManager.h"
+#include "Templates/SharedPointer.h"
+#include "Widgets/DeclarativeSyntaxSupport.h"
+#include "Widgets/SBoxPanel.h"
+#include "Widgets/SCompoundWidget.h"
+#include "Widgets/SWidget.h"
+
+class STORAGESERVERWIDGETS_API SZenProjectStatistics : public SCompoundWidget
+{
+	SLATE_BEGIN_ARGS(SZenProjectStatistics)
+		: _ZenServiceInstance(nullptr)
+	{ }
+
+	SLATE_ATTRIBUTE(TSharedPtr<UE::Zen::FZenServiceInstance>, ZenServiceInstance);
+
+	SLATE_END_ARGS()
+
+	void Construct(const FArguments& InArgs);
+
+private:
+
+	TSharedRef<SWidget> GetGridPanel();
+
+	EActiveTimerReturnType UpdateGridPanels(double InCurrentTime, float InDeltaTime);
+
+	SVerticalBox::FSlot* GridSlot = nullptr;
+	TAttribute<TSharedPtr<UE::Zen::FZenServiceInstance>> ZenServiceInstance;
+};
\ No newline at end of file
diff --git a/Engine/Source/Editor/StorageServerWidgets/Internal/SZenServiceStatus.h b/Engine/Source/Editor/StorageServerWidgets/Internal/SZenServiceStatus.h
new file mode 100644
--- /dev/null
+++ b/Engine/Source/Editor/StorageServerWidgets/Internal/SZenServiceStatus.h
@@ -0,0 +1,51 @@
+// Copyright Epic Games, Inc. All Rights Reserved.
+
+#pragma once
+
+#include "Experimental/ZenServerInterface.h"
+#include "ZenServiceInstanceManager.h"
+#include "Templates/SharedPointer.h"
+#include "Widgets/DeclarativeSyntaxSupport.h"
+#include "Widgets/SBoxPanel.h"
+#include "Widgets/SCompoundWidget.h"
+#include "Widgets/SWidget.h"
+
+class STORAGESERVERWIDGETS_API SZenServiceStatus : public SCompoundWidget
+{
+	SLATE_BEGIN_ARGS(SZenServiceStatus)
+		: _ZenServiceInstance(nullptr)
+	{ }
+
+	SLATE_ATTRIBUTE(TSharedPtr<UE::Zen::FZenServiceInstance>, ZenServiceInstance);
+
+	SLATE_END_ARGS()
+
+	void Construct(const FArguments& InArgs);
+
+private:
+	struct FState
+	{
+		UE::Zen::FZenLocalServiceRunContext RunContext;
+		FString Version;
+		UE::Zen::FGCStatus GCStatus;
+		UE::Zen::FZenCacheStats ZenCacheStats;
+		UE::Zen::FZenProjectStats ZenProjectStats;
+		uint16 LocalPort = 0;
+		bool bHaveStats = false;
+		bool bGotRunContext = false;
+		bool bIsRunning = false;
+	};
+	static constexpr uint32 NumState = 2;
+	FState State[NumState];
+	std::atomic<uint32> ActiveStateIndex = 0;
+
+
+	const FState& GetCurrentState() const;
+	TSharedRef<SWidget> GetGridPanel();
+	FReply ExploreDataPath_OnClicked();
+
+	EActiveTimerReturnType UpdateState(double InCurrentTime, float InDeltaTime);
+
+	SVerticalBox::FSlot* GridSlot = nullptr;
+	TAttribute<TSharedPtr<UE::Zen::FZenServiceInstance>> ZenServiceInstance;
+};
\ No newline at end of file
diff --git a/Engine/Source/Editor/StorageServerWidgets/Internal/ZenServiceInstanceManager.h b/Engine/Source/Editor/StorageServerWidgets/Internal/ZenServiceInstanceManager.h
new file mode 100644
--- /dev/null
+++ b/Engine/Source/Editor/StorageServerWidgets/Internal/ZenServiceInstanceManager.h
@@ -0,0 +1,22 @@
+// Copyright Epic Games, Inc. All Rights Reserved.
+
+#pragma once
+
+#include "Templates/SharedPointer.h"
+
+namespace UE::Zen { class FZenServiceInstance; }
+
+namespace UE::Zen
+{
+
+class STORAGESERVERWIDGETS_API FServiceInstanceManager
+{
+public:
+	TSharedPtr<FZenServiceInstance> GetZenServiceInstance() const;
+private:
+	mutable uint16 CurrentPort = 0;
+	mutable TSharedPtr<FZenServiceInstance> CurrentInstance;
+};
+
+} // namespace UE::Zen
+
diff --git a/Engine/Source/Editor/StorageServerWidgets/Private/SZenCacheStatistics.cpp b/Engine/Source/Editor/StorageServerWidgets/Private/SZenCacheStatistics.cpp
new file mode 100644
--- /dev/null
+++ b/Engine/Source/Editor/StorageServerWidgets/Private/SZenCacheStatistics.cpp
@@ -0,0 +1,313 @@
+// Copyright Epic Games, Inc. All Rights Reserved.
+
+#include "SZenCacheStatistics.h"
+#include "Experimental/ZenServerInterface.h"
+#include "Internationalization/FastDecimalFormat.h"
+#include "Math/BasicMathExpressionEvaluator.h"
+#include "Math/UnitConversion.h"
+#include "Misc/ExpressionParser.h"
+#include "Styling/StyleColors.h"
+#include "Widgets/Layout/SGridPanel.h"
+#include "Widgets/SBoxPanel.h"
+#include "Widgets/Text/STextBlock.h"
+
+#define LOCTEXT_NAMESPACE "ZenDashboard"
+
+void SZenCacheStatistics::Construct(const FArguments& InArgs)
+{
+	ZenServiceInstance = InArgs._ZenServiceInstance;
+
+	this->ChildSlot
+	[
+		SNew(SVerticalBox)
+
+		+ SVerticalBox::Slot()
+		.AutoHeight()
+		.Padding(0, 0, 0, 0)
+		.Expose(GridSlot)
+		[
+			GetGridPanel()
+		]
+	];
+
+	RegisterActiveTimer(0.5f, FWidgetActiveTimerDelegate::CreateSP(this, &SZenCacheStatistics::UpdateGridPanels));
+}
+
+EActiveTimerReturnType SZenCacheStatistics::UpdateGridPanels(double InCurrentTime, float InDeltaTime)
+{
+	(*GridSlot)
+	[
+		GetGridPanel()
+	];
+
+	SlatePrepass(GetPrepassLayoutScaleMultiplier());
+
+	return EActiveTimerReturnType::Continue;
+}
+
+TSharedRef<SWidget> SZenCacheStatistics::GetGridPanel()
+{
+	TSharedRef<SGridPanel> Panel = SNew(SGridPanel);
+
+	bool bHaveStats = false;
+	UE::Zen::FZenCacheStats ZenStats;
+
+	if (TSharedPtr<UE::Zen::FZenServiceInstance> ServiceInstance = ZenServiceInstance.Get())
+	{
+		if (ServiceInstance->GetCacheStats(ZenStats))
+		{
+			bHaveStats = true;
+		}
+	}
+
+	const static FNumberFormattingOptions SingleDecimalFormatting = FNumberFormattingOptions()
+		.SetUseGrouping(true)
+		.SetMinimumFractionalDigits(1)
+		.SetMaximumFractionalDigits(1);
+
+	int32 Row = 0;
+
+	const float RowMargin = 0.0f;
+	const float TitleMargin = 10.0f;
+	const float ColumnMargin = 10.0f;
+	const FSlateColor TitleColor = FStyleColors::AccentWhite;
+	const FSlateFontInfo TitleFont = FCoreStyle::GetDefaultFontStyle("Bold", 10);
+
+
+	Panel->AddSlot(0, Row)
+	.ColumnSpan(2)
+	[
+		SNew(STextBlock)
+		.Margin(FMargin(ColumnMargin, RowMargin))
+		.ColorAndOpacity(TitleColor)
+		.Font(TitleFont)
+		.Text(LOCTEXT("Cache", "Local Cache"))
+	];
+
+	Row++;
+
+	Panel->AddSlot(0, Row)
+	[
+		SNew(STextBlock)
+		.Margin(FMargin(ColumnMargin, RowMargin))
+		.Text(LOCTEXT("CacheHitRate", "Hit Rate"))
+	];
+
+	Panel->AddSlot(1, Row)
+	[
+		SNew(STextBlock)
+		.Margin(FMargin(ColumnMargin, RowMargin))
+		.Text_Lambda([CacheStats = ZenStats.General]
+		{
+			int64 CacheTotal = CacheStats.Hits + CacheStats.Misses;
+			return CacheTotal == 0 ? LOCTEXT("CacheNoHitRateValue", "-") : FText::AsPercent(CacheStats.HitRatio, &SingleDecimalFormatting);
+		})
+	];
+
+	Row++;
+
+	Panel->AddSlot(0, Row)
+	[
+		SNew(STextBlock)
+		.Margin(FMargin(ColumnMargin, RowMargin))
+		.Text(LOCTEXT("CacheHitQuantity", "Hits"))
+	];
+
+	Panel->AddSlot(1, Row)
+	[
+		SNew(STextBlock)
+		.Margin(FMargin(ColumnMargin, RowMargin))
+		.Text_Lambda([bHaveStats, CacheHits = ZenStats.General.Hits]
+		{
+			return bHaveStats ? FText::AsNumber(CacheHits) : LOCTEXT("UnavailableValue", "-");
+		})
+	];
+
+	Row++;
+
+	Panel->AddSlot(0, Row)
+	[
+		SNew(STextBlock)
+		.Margin(FMargin(ColumnMargin, RowMargin))
+		.Text(LOCTEXT("CacheMissQuantity", "Misses"))
+	];
+
+	Panel->AddSlot(1, Row)
+	[
+		SNew(STextBlock)
+		.Margin(FMargin(ColumnMargin, RowMargin))
+		.Text_Lambda([bHaveStats, CacheMisses = ZenStats.General.Misses]
+		{
+			return bHaveStats ? FText::AsNumber(CacheMisses) : LOCTEXT("UnavailableValue", "-");
+		})
+	];
+
+	Row++;
+
+	Panel->AddSlot(0, Row)
+	[
+		SNew(STextBlock)
+		.Margin(FMargin(ColumnMargin, RowMargin))
+		.Text(LOCTEXT("CacheWriteQuantity", "Writes"))
+	];
+
+	Panel->AddSlot(1, Row)
+	[
+		SNew(STextBlock)
+		.Margin(FMargin(ColumnMargin, RowMargin))
+		.Text_Lambda([bHaveStats, CacheWrites = ZenStats.General.Writes]
+		{
+			return bHaveStats ? FText::AsNumber(CacheWrites) : LOCTEXT("UnavailableValue", "-");
+		})
+	];
+
+	Row++;
+
+	Panel->AddSlot(0, Row)
+	[
+		SNew(STextBlock)
+		.Margin(FMargin(ColumnMargin, RowMargin))
+		.Text(LOCTEXT("CacheRequests", "Requests"))
+	];
+
+	Panel->AddSlot(1, Row)
+	[
+		SNew(STextBlock)
+		.Margin(FMargin(ColumnMargin, RowMargin))
+		.Text_Lambda([bHaveStats, Requests = ZenStats.Request.Count]
+		{
+			if (bHaveStats)
+			{
+				return FText::AsNumber(Requests);
+			}
+			return LOCTEXT("UnavailableValue", "-");
+		})
+	];
+
+	Row++;
+
+	Panel->AddSlot(0, Row)
+	[
+		SNew(STextBlock)
+		.Margin(FMargin(ColumnMargin, RowMargin))
+		.Text(LOCTEXT("CacheBadRequests", "Bad Requests"))
+	];
+
+	Panel->AddSlot(1, Row)
+	[
+		SNew(STextBlock)
+		.Margin(FMargin(ColumnMargin, RowMargin))
+		.Text_Lambda([bHaveStats, BadRequests = ZenStats.General.BadRequestCount]
+		{
+			if (bHaveStats)
+			{
+				return FText::AsNumber(BadRequests);
+			}
+			return LOCTEXT("UnavailableValue", "-");
+		})
+	];
+
+	Row++;
+
+	if (!ZenStats.Upstream.EndPoint.IsEmpty())
+	{
+		Panel->AddSlot(0, Row)
+		[
+			SNew(STextBlock)
+			.Margin(FMargin(ColumnMargin, TitleMargin, ColumnMargin, 0.0f))
+			.ColorAndOpacity(TitleColor)
+			.Font(TitleFont)
+			.Text(LOCTEXT("UpstreamServer", "Upstream Server"))
+		];
+
+		Panel->AddSlot(1, Row)
+		[
+			SNew(STextBlock)
+			.Margin(FMargin(ColumnMargin, TitleMargin, ColumnMargin, 0.0f))
+			.ColorAndOpacity(TitleColor)
+			.Font(TitleFont)
+			.Text(LOCTEXT("HitRate", "Hit Rate"))
+		];
+
+		Panel->AddSlot(2, Row)
+		[
+			SNew(STextBlock)
+			.Margin(FMargin(ColumnMargin, TitleMargin, ColumnMargin, 0.0f))
+			.ColorAndOpacity(TitleColor)
+			.Font(TitleFont)
+			.Text(LOCTEXT("Downloaded", "Downloaded"))
+		];
+
+		Panel->AddSlot(3, Row)
+		[
+			SNew(STextBlock)
+			.Margin(FMargin(ColumnMargin, TitleMargin, ColumnMargin, 0.0f))
+			.ColorAndOpacity(TitleColor)
+			.Font(TitleFont)
+			.Text(LOCTEXT("Uploaded", "Uploaded"))
+		];
+
+		Panel->AddSlot(4, Row)
+		[
+			SNew(STextBlock)
+			.Margin(FMargin(ColumnMargin, TitleMargin, ColumnMargin, 0.0f))
+			.ColorAndOpacity(TitleColor)
+			.Font(TitleFont)
+			.Text(LOCTEXT("URL", "URL"))
+		];
+		Row++;
+
+		for (const UE::Zen::FZenCacheStats::FEndPointStats& EndpointStats : ZenStats.Upstream.EndPoint)
+		{
+			Panel->AddSlot(0, Row)
+			[
+				SNew(STextBlock)
+				.Margin(FMargin(ColumnMargin, RowMargin))
+				.Text_Lambda([EndpointStats] { return FText::FromString(EndpointStats.Name); })
+			];
+
+			Panel->AddSlot(1, Row)
+			[
+				SNew(STextBlock)
+				.Margin(FMargin(ColumnMargin, RowMargin))
+				.Text_Lambda([EndpointStats] { return FText::AsPercent(EndpointStats.HitRatio, &SingleDecimalFormatting); })
+			];
+
+			Panel->AddSlot(2, Row)
+			[
+				SNew(STextBlock)
+				.Margin(FMargin(ColumnMargin, RowMargin))
+				.Text_Lambda([DownloadedMB = EndpointStats.DownloadedMB]
+					{
+						uint64 DownloadedBytes = uint64(DownloadedMB*1024*1024);
+						return FText::AsMemory(DownloadedBytes, (DownloadedBytes > 1024) ? &SingleDecimalFormatting : nullptr, nullptr, EMemoryUnitStandard::IEC);
+					})
+												
+			];
+
+			Panel->AddSlot(3, Row)
+			[
+				SNew(STextBlock)
+				.Margin(FMargin(ColumnMargin, RowMargin))
+				.Text_Lambda([UploadedMB = EndpointStats.UploadedMB]
+					{
+						uint64 UploadedBytes = uint64(UploadedMB*1024*1024);
+						return FText::AsMemory(UploadedBytes, (UploadedBytes > 1024) ? &SingleDecimalFormatting : nullptr, nullptr, EMemoryUnitStandard::IEC);
+					})
+			];
+
+			Panel->AddSlot(4, Row)
+			[
+				SNew(STextBlock)
+				.Margin(FMargin(ColumnMargin, RowMargin))
+				.Text_Lambda([EndpointStats] { return FText::FromString(EndpointStats.Url); })
+			];
+
+			Row++;
+		}
+	}
+
+	return Panel;
+}
+
+#undef LOCTEXT_NAMESPACE
diff --git a/Engine/Source/Editor/StorageServerWidgets/Private/SZenCidStoreStatistics.cpp b/Engine/Source/Editor/StorageServerWidgets/Private/SZenCidStoreStatistics.cpp
new file mode 100644
--- /dev/null
+++ b/Engine/Source/Editor/StorageServerWidgets/Private/SZenCidStoreStatistics.cpp
@@ -0,0 +1,213 @@
+// Copyright Epic Games, Inc. All Rights Reserved.
+
+#include "SZenCidStoreStatistics.h"
+#include "Experimental/ZenServerInterface.h"
+#include "Internationalization/FastDecimalFormat.h"
+#include "Math/BasicMathExpressionEvaluator.h"
+#include "Math/UnitConversion.h"
+#include "Misc/ExpressionParser.h"
+#include "Styling/StyleColors.h"
+#include "Widgets/Layout/SGridPanel.h"
+#include "Widgets/SBoxPanel.h"
+#include "Widgets/Text/STextBlock.h"
+
+#define LOCTEXT_NAMESPACE "ZenDashboard"
+
+void SZenCidStoreStatistics::Construct(const FArguments& InArgs)
+{
+	ZenServiceInstance = InArgs._ZenServiceInstance;
+
+	this->ChildSlot
+	[
+		SNew(SVerticalBox)
+
+		+ SVerticalBox::Slot()
+		.AutoHeight()
+		.Padding(0, 0, 0, 0)
+		.Expose(GridSlot)
+		[
+			GetGridPanel()
+		]
+	];
+
+	RegisterActiveTimer(0.5f, FWidgetActiveTimerDelegate::CreateSP(this, &SZenCidStoreStatistics::UpdateGridPanels));
+}
+
+EActiveTimerReturnType SZenCidStoreStatistics::UpdateGridPanels(double InCurrentTime, float InDeltaTime)
+{
+	(*GridSlot)
+	[
+		GetGridPanel()
+	];
+
+	SlatePrepass(GetPrepassLayoutScaleMultiplier());
+
+	return EActiveTimerReturnType::Continue;
+}
+
+TSharedRef<SWidget> SZenCidStoreStatistics::GetGridPanel()
+{
+	TSharedRef<SGridPanel> Panel = SNew(SGridPanel);
+
+	bool bHaveStats = false;
+	UE::Zen::FZenCacheStats ZenStats;
+
+	if (TSharedPtr<UE::Zen::FZenServiceInstance> ServiceInstance = ZenServiceInstance.Get())
+	{
+		if (ServiceInstance->GetCacheStats(ZenStats))
+		{
+			bHaveStats = true;
+		}
+	}
+
+	int64 TotalHits = ZenStats.General.CidHits;
+	int64 TotalMisses = ZenStats.General.CidMisses;
+	int64 TotalWrites = ZenStats.General.CidWrites;
+	int64 Requests = ZenStats.General.CidHits + ZenStats.General.CidMisses + ZenStats.General.CidWrites;
+
+	const static FNumberFormattingOptions SingleDecimalFormatting = FNumberFormattingOptions()
+		.SetUseGrouping(true)
+		.SetMinimumFractionalDigits(1)
+		.SetMaximumFractionalDigits(1);
+
+	int32 Row = 0;
+
+	const float RowMargin = 0.0f;
+	const float TitleMargin = 10.0f;
+	const float ColumnMargin = 10.0f;
+	const FSlateColor TitleColor = FStyleColors::AccentWhite;
+	const FSlateFontInfo TitleFont = FCoreStyle::GetDefaultFontStyle("Bold", 10);
+
+
+	Panel->AddSlot(0, Row)
+	.ColumnSpan(2)
+	[
+		SNew(STextBlock)
+		.Margin(FMargin(ColumnMargin, RowMargin))
+		.ColorAndOpacity(TitleColor)
+		.Font(TitleFont)
+		.Text(LOCTEXT("CASCid", "Local Cid Store"))
+	];
+
+	Row++;
+
+	Panel->AddSlot(0, Row)
+	[
+		SNew(STextBlock)
+		.Margin(FMargin(ColumnMargin, RowMargin))
+		.Text(LOCTEXT("CidStoreHitRate", "Hit Rate"))
+	];
+
+	Panel->AddSlot(1, Row)
+	[
+		SNew(STextBlock)
+		.Margin(FMargin(ColumnMargin, RowMargin))
+		.Text_Lambda([bHaveStats, TotalHits, TotalMisses]
+		{
+			if (bHaveStats && ((TotalHits+TotalMisses) > 0))
+			{
+				double HitRate = static_cast<double>(TotalHits) / static_cast<double>(TotalHits + TotalMisses);
+				return FText::AsPercent(HitRate, &SingleDecimalFormatting);
+			}
+			return LOCTEXT("UnavailableValue", "-");
+		})
+	];
+
+	Row++;
+
+	Panel->AddSlot(0, Row)
+	[
+		SNew(STextBlock)
+		.Margin(FMargin(ColumnMargin, RowMargin))
+		.Text(LOCTEXT("CidStoreHitQuantity", "Hits"))
+	];
+
+	Panel->AddSlot(1, Row)
+	[
+		SNew(STextBlock)
+		.Margin(FMargin(ColumnMargin, RowMargin))
+		.Text_Lambda([bHaveStats, TotalHits]
+		{
+			if (bHaveStats)
+			{
+				return FText::AsNumber(TotalHits);
+			}
+			return LOCTEXT("UnavailableValue", "-");
+		})
+	];
+
+	Row++;
+
+	Panel->AddSlot(0, Row)
+	[
+		SNew(STextBlock)
+		.Margin(FMargin(ColumnMargin, RowMargin))
+		.Text(LOCTEXT("CidStoreMissQuantity", "Misses"))
+	];
+
+	Panel->AddSlot(1, Row)
+	[
+		SNew(STextBlock)
+		.Margin(FMargin(ColumnMargin, RowMargin))
+		.Text_Lambda([bHaveStats, TotalMisses]
+		{
+			if (bHaveStats)
+			{
+				return FText::AsNumber(TotalMisses);
+			}
+			return LOCTEXT("UnavailableValue", "-");
+		})
+	];
+
+	Row++;
+
+	Panel->AddSlot(0, Row)
+	[
+		SNew(STextBlock)
+		.Margin(FMargin(ColumnMargin, RowMargin))
+		.Text(LOCTEXT("CidStoreWriteQuantity", "Writes"))
+	];
+
+	Panel->AddSlot(1, Row)
+	[
+		SNew(STextBlock)
+		.Margin(FMargin(ColumnMargin, RowMargin))
+		.Text_Lambda([bHaveStats, TotalWrites]
+		{
+			if (bHaveStats)
+			{
+				return FText::AsNumber(TotalWrites);
+			}
+			return LOCTEXT("UnavailableValue", "-");
+		})
+	];
+
+	Row++;
+
+	Panel->AddSlot(0, Row)
+	[
+		SNew(STextBlock)
+		.Margin(FMargin(ColumnMargin, RowMargin))
+		.Text(LOCTEXT("ProjectRequests", "Requests"))
+	];
+
+	Panel->AddSlot(1, Row)
+	[
+		SNew(STextBlock)
+		.Margin(FMargin(ColumnMargin, RowMargin))
+		.Text_Lambda([bHaveStats, Requests]
+		{
+			if (bHaveStats)
+			{
+				return FText::AsNumber(Requests);
+			}
+			return LOCTEXT("UnavailableValue", "-");
+		})
+	];
+
+	Row++;
+
+	return Panel;
+}
+
+#undef LOCTEXT_NAMESPACE
diff --git a/Engine/Source/Editor/StorageServerWidgets/Private/SZenDialogs.cpp b/Engine/Source/Editor/StorageServerWidgets/Private/SZenDialogs.cpp
new file mode 100644
--- /dev/null
+++ b/Engine/Source/Editor/StorageServerWidgets/Private/SZenDialogs.cpp
@@ -0,0 +1,74 @@
+// Copyright Epic Games, Inc. All Rights Reserved.
+
+#include "SZenDialogs.h"
+
+#include "ZenServiceInstanceManager.h"
+#include "SZenCacheStatistics.h"
+#include "SZenCidStoreStatistics.h"
+#include "SZenProjectStatistics.h"
+#include "SZenServiceStatus.h"
+
+#include "Framework/Application/SlateApplication.h"
+#include "Styling/StyleColors.h"
+#include "Widgets/Images/SImage.h"
+#include "Widgets/Layout/SGridPanel.h"
+#include "Widgets/SBoxPanel.h"
+#include "Widgets/SToolTip.h"
+#include "Widgets/Text/STextBlock.h"
+
+#define LOCTEXT_NAMESPACE "ZenEditor"
+
+void SZenStoreStausDialog::Construct(const FArguments& InArgs)
+{
+	TSharedPtr<UE::Zen::FServiceInstanceManager> ServiceInstanceManager = MakeShared<UE::Zen::FServiceInstanceManager>();
+
+	this->ChildSlot
+		[
+			SNew(SVerticalBox)
+				+ SVerticalBox::Slot()
+				.AutoHeight()
+				.Padding(0.0f, 10.0f, 0.0f, 0.0f)
+				.VAlign(VAlign_Top)
+				[
+					SNew(SZenServiceStatus)
+						.ZenServiceInstance(ServiceInstanceManager.ToSharedRef(), &UE::Zen::FServiceInstanceManager::GetZenServiceInstance)
+				]
+				// Stats panel
+				+ SVerticalBox::Slot()
+				.AutoHeight()
+				.Padding(0.0f, 10.0f, 0.0f, 0.0f)
+				.VAlign(VAlign_Top)
+				[
+					SNew(SHorizontalBox)
+
+						+ SHorizontalBox::Slot()
+						.AutoWidth()
+						.HAlign(HAlign_Left)
+						.VAlign(VAlign_Top)
+						[
+							SNew(SZenCacheStatistics)
+								.ZenServiceInstance(ServiceInstanceManager.ToSharedRef(), &UE::Zen::FServiceInstanceManager::GetZenServiceInstance)
+						]
+
+						+ SHorizontalBox::Slot()
+						.AutoWidth()
+						.HAlign(HAlign_Left)
+						.VAlign(VAlign_Top)
+						[
+							SNew(SZenProjectStatistics)
+								.ZenServiceInstance(ServiceInstanceManager.ToSharedRef(), &UE::Zen::FServiceInstanceManager::GetZenServiceInstance)
+						]
+
+						+ SHorizontalBox::Slot()
+						.AutoWidth()
+						.HAlign(HAlign_Left)
+						.VAlign(VAlign_Top)
+						[
+							SNew(SZenCidStoreStatistics)
+								.ZenServiceInstance(ServiceInstanceManager.ToSharedRef(), &UE::Zen::FServiceInstanceManager::GetZenServiceInstance)
+						]
+				]
+		];
+}
+
+#undef LOCTEXT_NAMESPACE
diff --git a/Engine/Source/Editor/StorageServerWidgets/Private/SZenProjectStatistics.cpp b/Engine/Source/Editor/StorageServerWidgets/Private/SZenProjectStatistics.cpp
new file mode 100644
--- /dev/null
+++ b/Engine/Source/Editor/StorageServerWidgets/Private/SZenProjectStatistics.cpp
@@ -0,0 +1,235 @@
+// Copyright Epic Games, Inc. All Rights Reserved.
+
+#include "SZenProjectStatistics.h"
+#include "Experimental/ZenServerInterface.h"
+#include "Internationalization/FastDecimalFormat.h"
+#include "Math/BasicMathExpressionEvaluator.h"
+#include "Math/UnitConversion.h"
+#include "Misc/ExpressionParser.h"
+#include "Styling/StyleColors.h"
+#include "Widgets/Layout/SGridPanel.h"
+#include "Widgets/SBoxPanel.h"
+#include "Widgets/Text/STextBlock.h"
+
+#define LOCTEXT_NAMESPACE "ZenDashboard"
+
+void SZenProjectStatistics::Construct(const FArguments& InArgs)
+{
+	ZenServiceInstance = InArgs._ZenServiceInstance;
+
+	this->ChildSlot
+	[
+		SNew(SVerticalBox)
+
+		+ SVerticalBox::Slot()
+		.AutoHeight()
+		.Padding(0, 0, 0, 0)
+		.Expose(GridSlot)
+		[
+			GetGridPanel()
+		]
+	];
+
+	RegisterActiveTimer(0.5f, FWidgetActiveTimerDelegate::CreateSP(this, &SZenProjectStatistics::UpdateGridPanels));
+}
+
+EActiveTimerReturnType SZenProjectStatistics::UpdateGridPanels(double InCurrentTime, float InDeltaTime)
+{
+	(*GridSlot)
+	[
+		GetGridPanel()
+	];
+
+	SlatePrepass(GetPrepassLayoutScaleMultiplier());
+
+	return EActiveTimerReturnType::Continue;
+}
+
+TSharedRef<SWidget> SZenProjectStatistics::GetGridPanel()
+{
+	TSharedRef<SGridPanel> Panel = SNew(SGridPanel);
+
+	bool bHaveStats = false;
+	UE::Zen::FZenProjectStats ZenStats;
+
+	if (TSharedPtr<UE::Zen::FZenServiceInstance> ServiceInstance = ZenServiceInstance.Get())
+	{
+		if (ServiceInstance->GetProjectStats(ZenStats))
+		{
+			bHaveStats = true;
+		}
+	}
+
+	int64 TotalHits = ZenStats.General.Op.HitCount + ZenStats.General.Chunk.HitCount;
+	int64 TotalMisses = ZenStats.General.Op.MissCount + ZenStats.General.Chunk.MissCount;
+	int64 TotalWrites = ZenStats.General.Op.WriteCount + ZenStats.General.Chunk.WriteCount;
+
+	const static FNumberFormattingOptions SingleDecimalFormatting = FNumberFormattingOptions()
+		.SetUseGrouping(true)
+		.SetMinimumFractionalDigits(1)
+		.SetMaximumFractionalDigits(1);
+
+	int32 Row = 0;
+
+	const float RowMargin = 0.0f;
+	const float TitleMargin = 10.0f;
+	const float ColumnMargin = 10.0f;
+	const FSlateColor TitleColor = FStyleColors::AccentWhite;
+	const FSlateFontInfo TitleFont = FCoreStyle::GetDefaultFontStyle("Bold", 10);
+
+
+	Panel->AddSlot(0, Row)
+	.ColumnSpan(2)
+	[
+		SNew(STextBlock)
+		.Margin(FMargin(ColumnMargin, RowMargin))
+		.ColorAndOpacity(TitleColor)
+		.Font(TitleFont)
+		.Text(LOCTEXT("CAS", "Local Project Store"))
+	];
+
+	Row++;
+
+	Panel->AddSlot(0, Row)
+	[
+		SNew(STextBlock)
+		.Margin(FMargin(ColumnMargin, RowMargin))
+		.Text(LOCTEXT("ProjectHitRate", "Hit Rate"))
+	];
+
+	Panel->AddSlot(1, Row)
+	[
+		SNew(STextBlock)
+		.Margin(FMargin(ColumnMargin, RowMargin))
+		.Text_Lambda([bHaveStats, TotalHits, TotalMisses]
+		{
+			if (bHaveStats && ((TotalHits+TotalMisses) > 0))
+			{
+				double HitRate = static_cast<double>(TotalHits) / static_cast<double>(TotalHits + TotalMisses);
+				return FText::AsPercent(HitRate, &SingleDecimalFormatting);
+			}
+			return LOCTEXT("UnavailableValue", "-");
+		})
+	];
+
+	Row++;
+
+	Panel->AddSlot(0, Row)
+	[
+		SNew(STextBlock)
+		.Margin(FMargin(ColumnMargin, RowMargin))
+		.Text(LOCTEXT("ProjectHitQuantity", "Hits"))
+	];
+
+	Panel->AddSlot(1, Row)
+	[
+		SNew(STextBlock)
+		.Margin(FMargin(ColumnMargin, RowMargin))
+		.Text_Lambda([bHaveStats, TotalHits]
+		{
+			if (bHaveStats)
+			{
+				return FText::AsNumber(TotalHits);
+			}
+			return LOCTEXT("UnavailableValue", "-");
+		})
+	];
+
+	Row++;
+
+	Panel->AddSlot(0, Row)
+	[
+		SNew(STextBlock)
+		.Margin(FMargin(ColumnMargin, RowMargin))
+		.Text(LOCTEXT("ProjectMissQuantity", "Misses"))
+	];
+
+	Panel->AddSlot(1, Row)
+	[
+		SNew(STextBlock)
+		.Margin(FMargin(ColumnMargin, RowMargin))
+		.Text_Lambda([bHaveStats, TotalMisses]
+		{
+			if (bHaveStats)
+			{
+				return FText::AsNumber(TotalMisses);
+			}
+			return LOCTEXT("UnavailableValue", "-");
+		})
+	];
+
+	Row++;
+
+	Panel->AddSlot(0, Row)
+	[
+		SNew(STextBlock)
+		.Margin(FMargin(ColumnMargin, RowMargin))
+		.Text(LOCTEXT("ProjectWriteQuantity", "Writes"))
+	];
+
+	Panel->AddSlot(1, Row)
+	[
+		SNew(STextBlock)
+		.Margin(FMargin(ColumnMargin, RowMargin))
+		.Text_Lambda([bHaveStats, TotalWrites]
+		{
+			if (bHaveStats)
+			{
+				return FText::AsNumber(TotalWrites);
+			}
+			return LOCTEXT("UnavailableValue", "-");
+		})
+	];
+
+	Row++;
+
+	Panel->AddSlot(0, Row)
+	[
+		SNew(STextBlock)
+		.Margin(FMargin(ColumnMargin, RowMargin))
+		.Text(LOCTEXT("ProjectRequests", "Requests"))
+	];
+
+	Panel->AddSlot(1, Row)
+	[
+		SNew(STextBlock)
+		.Margin(FMargin(ColumnMargin, RowMargin))
+		.Text_Lambda([bHaveStats, Requests = ZenStats.General.RequestCount]
+		{
+			if (bHaveStats)
+			{
+				return FText::AsNumber(Requests);
+			}
+			return LOCTEXT("UnavailableValue", "-");
+		})
+	];
+
+	Row++;
+
+	Panel->AddSlot(0, Row)
+		[
+			SNew(STextBlock)
+			.Margin(FMargin(ColumnMargin, RowMargin))
+		.Text(LOCTEXT("ProjectBadRequests", "Bad Requests"))
+		];
+
+	Panel->AddSlot(1, Row)
+		[
+			SNew(STextBlock)
+			.Margin(FMargin(ColumnMargin, RowMargin))
+		.Text_Lambda([bHaveStats, BadRequests = ZenStats.General.BadRequestCount]
+			{
+				if (bHaveStats)
+				{
+					return FText::AsNumber(BadRequests);
+				}
+	return LOCTEXT("UnavailableValue", "-");
+			})
+		];
+
+	Row++;
+
+	return Panel;
+}
+
+#undef LOCTEXT_NAMESPACE
diff --git a/Engine/Source/Editor/StorageServerWidgets/Private/SZenServiceStatus.cpp b/Engine/Source/Editor/StorageServerWidgets/Private/SZenServiceStatus.cpp
new file mode 100644
--- /dev/null
+++ b/Engine/Source/Editor/StorageServerWidgets/Private/SZenServiceStatus.cpp
@@ -0,0 +1,280 @@
+// Copyright Epic Games, Inc. All Rights Reserved.
+
+#include "SZenServiceStatus.h"
+#include "Experimental/ZenServerInterface.h"
+#include "Internationalization/FastDecimalFormat.h"
+#include "Math/BasicMathExpressionEvaluator.h"
+#include "Math/UnitConversion.h"
+#include "Misc/ExpressionParser.h"
+#include "Styling/StyleColors.h"
+#include "Widgets/Images/SImage.h"
+#include "Widgets/Input/SButton.h"
+#include "Widgets/Layout/SGridPanel.h"
+#include "Widgets/SBoxPanel.h"
+#include "Widgets/Text/STextBlock.h"
+
+#define LOCTEXT_NAMESPACE "ZenDashboard"
+
+void SZenServiceStatus::Construct(const FArguments& InArgs)
+{
+	ZenServiceInstance = InArgs._ZenServiceInstance;
+
+	this->ChildSlot
+	[
+		SNew(SVerticalBox)
+
+		+ SVerticalBox::Slot()
+		.AutoHeight()
+		.Padding(0, 0, 0, 0)
+		.Expose(GridSlot)
+		[
+			GetGridPanel()
+		]
+	];
+
+	UpdateState(0.0, 0.0);
+	RegisterActiveTimer(0.5f, FWidgetActiveTimerDelegate::CreateSP(this, &SZenServiceStatus::UpdateState));
+}
+
+EActiveTimerReturnType SZenServiceStatus::UpdateState(double InCurrentTime, float InDeltaTime)
+{
+	uint32 LastStateIndex = ActiveStateIndex.load(std::memory_order_relaxed);
+	uint32 NextStateIndex = (LastStateIndex + 1)%NumState;
+	FState& NextState = State[NextStateIndex];
+	NextState.bGotRunContext = UE::Zen::TryGetLocalServiceRunContext(NextState.RunContext);
+	NextState.Version = UE::Zen::GetLocalServiceInstallVersion(false);
+	if (NextState.bGotRunContext)
+	{
+		NextState.bIsRunning = UE::Zen::IsLocalServiceRunning(*NextState.RunContext.GetDataPath(), &NextState.LocalPort);
+	}
+
+	NextState.GCStatus = UE::Zen::FGCStatus{};
+	NextState.bHaveStats = false;
+	if (NextState.bIsRunning)
+	{
+		if (TSharedPtr<UE::Zen::FZenServiceInstance> ServiceInstance = ZenServiceInstance.Get())
+		{
+			ServiceInstance->GetGCStatus(NextState.GCStatus);
+			if (ServiceInstance->GetCacheStats(NextState.ZenCacheStats) && ServiceInstance->GetProjectStats(NextState.ZenProjectStats))
+			{
+				NextState.bHaveStats = true;
+			}
+		}
+	}
+
+	ActiveStateIndex.compare_exchange_weak(LastStateIndex, NextStateIndex, std::memory_order_relaxed);
+
+	return EActiveTimerReturnType::Continue;
+}
+
+FReply SZenServiceStatus::ExploreDataPath_OnClicked()
+{
+	const FState& CurrentState = GetCurrentState();
+	if (CurrentState.bGotRunContext)
+	{
+		FString FullPath(FPaths::ConvertRelativePathToFull(CurrentState.RunContext.GetDataPath()));
+		FPlatformProcess::ExploreFolder(*FullPath);
+		return FReply::Handled();
+	}
+	return FReply::Unhandled();
+}
+
+const SZenServiceStatus::FState& SZenServiceStatus::GetCurrentState() const
+{
+	return State[ActiveStateIndex.load(std::memory_order_relaxed)];
+}
+
+TSharedRef<SWidget> SZenServiceStatus::GetGridPanel()
+{
+	TSharedRef<SGridPanel> Panel = SNew(SGridPanel);
+	
+	const static FNumberFormattingOptions SingleDecimalFormatting = FNumberFormattingOptions()
+		.SetUseGrouping(true)
+		.SetMinimumFractionalDigits(1)
+		.SetMaximumFractionalDigits(1);
+
+	int32 Row = 0;
+
+	const float RowMargin = 0.0f;
+	const float ColumnMargin = 10.0f;
+	const FSlateColor TitleColor = FStyleColors::AccentWhite;
+	const FSlateFontInfo TitleFont = FCoreStyle::GetDefaultFontStyle("Bold", 10);
+
+	Panel->AddSlot(0, Row)
+	[
+		SNew(STextBlock)
+		.Margin(FMargin(ColumnMargin, RowMargin))
+		.ColorAndOpacity(TitleColor)
+		.Font(TitleFont)
+		.Text(LOCTEXT("ServiceStatus_Status", "Status"))
+	];
+
+	Panel->AddSlot(1, Row)
+	[
+		SNew(STextBlock)
+		.Margin(FMargin(ColumnMargin, RowMargin))
+		.Text_Lambda([this]
+		{
+			const FState& CurrentState = GetCurrentState();
+			FText ServiceStatus = LOCTEXT("ServiceStatus_NotInstalled", "Not installed");
+			if (CurrentState.bGotRunContext)
+			{
+				if (CurrentState.bIsRunning && (CurrentState.LocalPort != 0))
+				{
+					ServiceStatus = LOCTEXT("ServiceStatus_Running", "Running");
+				}
+				else
+				{
+					ServiceStatus = LOCTEXT("ServiceStatus_Stopped", "Stopped");
+				}
+			}
+			return ServiceStatus;
+		})
+	];
+
+	++Row;
+
+	Panel->AddSlot(0, Row)
+	[
+		SNew(STextBlock)
+		.Margin(FMargin(ColumnMargin, RowMargin))
+		.ColorAndOpacity(TitleColor)
+		.Font(TitleFont)
+		.Text(LOCTEXT("ServiceStatus_Version", "Version"))
+	];
+
+	Panel->AddSlot(1, Row)
+	[
+		SNew(STextBlock)
+		.Margin(FMargin(ColumnMargin, RowMargin))
+		.Text_Lambda([this]
+		{
+			const FState& CurrentState = GetCurrentState();
+			return FText::FromString(CurrentState.Version);
+		})
+	];
+
+	++Row;
+
+	Panel->AddSlot(0, Row)
+	[
+		SNew(STextBlock)
+		.Margin(FMargin(ColumnMargin, RowMargin))
+		.ColorAndOpacity(TitleColor)
+		.Font(TitleFont)
+		.Text(LOCTEXT("ServiceStatus_Port", "Port"))
+	];
+
+	Panel->AddSlot(1, Row)
+	[
+		SNew(STextBlock)
+		.Margin(FMargin(ColumnMargin, RowMargin))
+		.Text_Lambda([this]
+		{
+			const FState& CurrentState = GetCurrentState();
+			return CurrentState.LocalPort == 0 ? LOCTEXT("ServiceStatus_NoPortValue", "-") : FText::AsNumber(CurrentState.LocalPort, &FNumberFormattingOptions::DefaultNoGrouping());
+		})
+	];
+
+	++Row;
+
+	Panel->AddSlot(0, Row)
+	[
+		SNew(STextBlock)
+		.Margin(FMargin(ColumnMargin, RowMargin))
+		.ColorAndOpacity(TitleColor)
+		.Font(TitleFont)
+		.Text(LOCTEXT("ServiceStatus_GCStatus", "GC Status"))
+	];
+
+	Panel->AddSlot(1, Row)
+	[
+		SNew(STextBlock)
+		.Margin(FMargin(ColumnMargin, RowMargin))
+		.Text_Lambda([this]
+		{
+			const FState& CurrentState = GetCurrentState();
+			return CurrentState.LocalPort == 0 ? LOCTEXT("ServiceStatus_NoGCStatusValue", "-") : FText::FromString(CurrentState.GCStatus.Description);
+		})
+	];
+
+	++Row;
+
+	Panel->AddSlot(0, Row)
+	[
+		SNew(STextBlock)
+		.Margin(FMargin(ColumnMargin, RowMargin))
+		.ColorAndOpacity(TitleColor)
+		.Font(TitleFont)
+		.Text(LOCTEXT("ServiceStatus_DiskSpace", "Disk Space"))
+	];
+
+	Panel->AddSlot(1, Row)
+	[
+		SNew(STextBlock)
+		.Margin(FMargin(ColumnMargin, RowMargin))
+		.Text_Lambda([this]
+		{
+			const FState& CurrentState = GetCurrentState();
+			const uint64 TotalDiskSpace = uint64(CurrentState.ZenCacheStats.General.Size.Disk) + uint64(CurrentState.ZenProjectStats.General.Size.Disk) + uint64(CurrentState.ZenCacheStats.CID.Size.Total);
+			return CurrentState.bHaveStats ? FText::AsMemory(TotalDiskSpace, (TotalDiskSpace > 1024) ? &SingleDecimalFormatting : nullptr, nullptr, EMemoryUnitStandard::IEC) : LOCTEXT("UnavailableValue", "-");
+		})
+	];
+
+	++Row;
+
+	Panel->AddSlot(0, Row)
+	.VAlign(VAlign_Bottom)
+	[
+		SNew(STextBlock)
+		.Margin(FMargin(ColumnMargin, RowMargin))
+		.ColorAndOpacity(TitleColor)
+		.Font(TitleFont)
+		.Text(LOCTEXT("ServiceStatus_DataPath", "Data path"))
+	];
+
+	Panel->AddSlot(1, Row)
+	[
+
+		SNew(SHorizontalBox)
+
+		+ SHorizontalBox::Slot()
+		.VAlign(VAlign_Bottom)
+		[
+			SNew(STextBlock)
+			.Margin(FMargin(ColumnMargin, RowMargin))
+			.Text_Lambda([this]
+			{
+				const FState& CurrentState = GetCurrentState();
+				return CurrentState.bGotRunContext ? FText::FromString(CurrentState.RunContext.GetDataPath()) : LOCTEXT("ServiceStatus_NoDataPathValue", "-");
+			})
+		]
+
+		+ SHorizontalBox::Slot()
+		.AutoWidth()
+		.VAlign(VAlign_Bottom)
+		[
+			SNew(SButton)
+			.ButtonStyle(&FAppStyle::Get().GetWidgetStyle<FButtonStyle>("SimpleButton"))
+			.ToolTipText(LOCTEXT("ExploreDataPathButtonToolTip", "Explore the Data Path"))
+			.IsEnabled_Lambda([this]
+			{
+				const FState& CurrentState = GetCurrentState();
+				return CurrentState.bGotRunContext;
+			})
+			.OnClicked(this, &SZenServiceStatus::ExploreDataPath_OnClicked)
+			[
+				SNew(SImage)
+				.Image(FAppStyle::Get().GetBrush("Zen.Icons.FolderExplore"))
+				.ColorAndOpacity(FSlateColor::UseForeground())
+			]
+		]
+	];
+
+	++Row;
+
+
+	return Panel;
+}
+
+#undef LOCTEXT_NAMESPACE
diff --git a/Engine/Source/Editor/StorageServerWidgets/Private/ZenServiceInstanceManager.cpp b/Engine/Source/Editor/StorageServerWidgets/Private/ZenServiceInstanceManager.cpp
new file mode 100644
--- /dev/null
+++ b/Engine/Source/Editor/StorageServerWidgets/Private/ZenServiceInstanceManager.cpp
@@ -0,0 +1,35 @@
+// Copyright Epic Games, Inc. All Rights Reserved.
+
+#include "ZenServiceInstanceManager.h"
+#include "Experimental/ZenServerInterface.h"
+
+
+namespace UE::Zen
+{
+
+TSharedPtr<FZenServiceInstance> FServiceInstanceManager::GetZenServiceInstance() const
+{
+	UE::Zen::FZenLocalServiceRunContext RunContext;
+	if (UE::Zen::TryGetLocalServiceRunContext(RunContext))
+	{
+		uint16 LocalPort = 0;
+		if (UE::Zen::IsLocalServiceRunning(*RunContext.GetDataPath(), &LocalPort) && (LocalPort != 0))
+		{
+			if (LocalPort != CurrentPort)
+			{
+				CurrentPort = LocalPort;
+				UE::Zen::FServiceSettings ServiceSettings;
+				ServiceSettings.SettingsVariant.Emplace<UE::Zen::FServiceConnectSettings>();
+				UE::Zen::FServiceConnectSettings& ConnectExistingSettings = ServiceSettings.SettingsVariant.Get<UE::Zen::FServiceConnectSettings>();
+				ConnectExistingSettings.HostName = TEXT("localhost");
+				ConnectExistingSettings.Port = LocalPort;
+				CurrentInstance = MakeShared<FZenServiceInstance>(MoveTemp(ServiceSettings));
+			}
+
+			return CurrentInstance;
+		}
+	}
+	return nullptr;
+}
+
+} // namespace UE::Zen
diff --git a/Engine/Source/Editor/StorageServerWidgets/StorageServerWidgets.Build.cs b/Engine/Source/Editor/StorageServerWidgets/StorageServerWidgets.Build.cs
new file mode 100644
--- /dev/null
+++ b/Engine/Source/Editor/StorageServerWidgets/StorageServerWidgets.Build.cs
@@ -0,0 +1,25 @@
+// Copyright Epic Games, Inc. All Rights Reserved.
+
+using UnrealBuildTool;
+
+public class StorageServerWidgets : ModuleRules
+{
+	public StorageServerWidgets(ReadOnlyTargetRules Target)
+		 : base(Target)
+	{
+		PrivateDependencyModuleNames.AddRange(
+			new string[] {
+				"ApplicationCore",
+				"Core",
+				"CoreUObject",
+				"InputCore",
+				"Json",
+				"Projects",	
+				"Slate",
+				"SlateCore",
+				"Zen",
+			});
+
+		CppCompileWarningSettings.UnsafeTypeCastWarningLevel = WarningLevel.Error;
+	}
+}
diff --git a/Engine/Source/Editor/ZenEditor/ZenEditor.Build.cs b/Engine/Source/Editor/ZenEditor/ZenEditor.Build.cs
--- a/Engine/Source/Editor/ZenEditor/ZenEditor.Build.cs
+++ b/Engine/Source/Editor/ZenEditor/ZenEditor.Build.cs
@@ -11,23 +11,23 @@ public class ZenEditor : ModuleRules
 			new string[] {
 				"Core",
 				"CoreUObject",
+				"DerivedDataCache",
+				"DerivedDataWidgets",
+				"EditorFramework",
+				"EditorSubsystem",
 				"Engine",
+				"InputCore",
+				"MessageLog",
+				"OutputLog",
 				"Slate",
 				"SlateCore",
-				"InputCore",
-				"EditorFramework",
+				"StorageServerWidgets",
+				"ToolMenus",
+				"ToolWidgets",
 				"UATHelper",
 				"UnrealEd",
-				"ToolMenus",
-				"OutputLog",
-				"DerivedDataCache",
-				"EditorSubsystem",
 				"WorkspaceMenuStructure",
-				"MessageLog",
-				"ToolWidgets",
-				"DerivedDataWidgets",
-				"Zen",
-				"ZenWidgets"
+				"Zen"
 			});
 
 		DynamicallyLoadedModuleNames.AddRange(
diff --git a/Engine/Source/Editor/ZenWidgets/Internal/SZenCacheStatistics.h b/Engine/Source/Editor/ZenWidgets/Internal/SZenCacheStatistics.h
deleted file mode 100644
diff --git a/Engine/Source/Editor/ZenWidgets/Internal/SZenCidStoreStatistics.h b/Engine/Source/Editor/ZenWidgets/Internal/SZenCidStoreStatistics.h
deleted file mode 100644
diff --git a/Engine/Source/Editor/ZenWidgets/Internal/SZenDialogs.h b/Engine/Source/Editor/ZenWidgets/Internal/SZenDialogs.h
deleted file mode 100644
diff --git a/Engine/Source/Editor/ZenWidgets/Internal/SZenProjectStatistics.h b/Engine/Source/Editor/ZenWidgets/Internal/SZenProjectStatistics.h
deleted file mode 100644
diff --git a/Engine/Source/Editor/ZenWidgets/Internal/SZenServiceStatus.h b/Engine/Source/Editor/ZenWidgets/Internal/SZenServiceStatus.h
deleted file mode 100644
diff --git a/Engine/Source/Editor/ZenWidgets/Internal/ZenServiceInstanceManager.h b/Engine/Source/Editor/ZenWidgets/Internal/ZenServiceInstanceManager.h
deleted file mode 100644
diff --git a/Engine/Source/Editor/ZenWidgets/Private/SZenCacheStatistics.cpp b/Engine/Source/Editor/ZenWidgets/Private/SZenCacheStatistics.cpp
deleted file mode 100644
diff --git a/Engine/Source/Editor/ZenWidgets/Private/SZenCidStoreStatistics.cpp b/Engine/Source/Editor/ZenWidgets/Private/SZenCidStoreStatistics.cpp
deleted file mode 100644
diff --git a/Engine/Source/Editor/ZenWidgets/Private/SZenDialogs.cpp b/Engine/Source/Editor/ZenWidgets/Private/SZenDialogs.cpp
deleted file mode 100644
diff --git a/Engine/Source/Editor/ZenWidgets/Private/SZenProjectStatistics.cpp b/Engine/Source/Editor/ZenWidgets/Private/SZenProjectStatistics.cpp
deleted file mode 100644
diff --git a/Engine/Source/Editor/ZenWidgets/Private/SZenServiceStatus.cpp b/Engine/Source/Editor/ZenWidgets/Private/SZenServiceStatus.cpp
deleted file mode 100644
diff --git a/Engine/Source/Editor/ZenWidgets/Private/ZenServiceInstanceManager.cpp b/Engine/Source/Editor/ZenWidgets/Private/ZenServiceInstanceManager.cpp
deleted file mode 100644
diff --git a/Engine/Source/Editor/ZenWidgets/ZenWidgets.Build.cs b/Engine/Source/Editor/ZenWidgets/ZenWidgets.Build.cs
deleted file mode 100644
diff --git a/Engine/Source/Programs/ZenDashboard/ZenDashboard.Build.cs b/Engine/Source/Programs/ZenDashboard/ZenDashboard.Build.cs
--- a/Engine/Source/Programs/ZenDashboard/ZenDashboard.Build.cs
+++ b/Engine/Source/Programs/ZenDashboard/ZenDashboard.Build.cs
@@ -7,19 +7,19 @@ public class ZenDashboard : ModuleRules
 	public ZenDashboard(ReadOnlyTargetRules Target) : base(Target)
 	{
 		bTreatAsEngineModule = true;
-		
+
 		PrivateDependencyModuleNames.AddRange(
 			new string[]
 			{
+				"ApplicationCore",
 				"Core",
 				"Json",
 				"Projects",
-				"ApplicationCore",
 				"Slate",
 				"SlateCore",
 				"StandaloneRenderer",
-				"Zen",
-				"ZenWidgets"
+				"StorageServerWidgets",
+				"Zen"
 			});
 
 		if (Target.IsInPlatformGroup(UnrealPlatformGroup.Linux))
