Zen dashboard usability improvements.
Allowing the installation directory to be browsed.  Added a "Browse" field with a link to the zen web UI.

diff --git a/Engine/Source/Editor/StorageServerWidgets/Internal/SZenServiceStatus.h b/Engine/Source/Editor/StorageServerWidgets/Internal/SZenServiceStatus.h
--- a/Engine/Source/Editor/StorageServerWidgets/Internal/SZenServiceStatus.h
+++ b/Engine/Source/Editor/StorageServerWidgets/Internal/SZenServiceStatus.h
@@ -43,6 +43,7 @@ private:
 	const FState& GetCurrentState() const;
 	TSharedRef<SWidget> GetGridPanel();
 	FReply ExploreDataPath_OnClicked();
+	FReply ExploreExecutablePath_OnClicked();
 
 	EActiveTimerReturnType UpdateState(double InCurrentTime, float InDeltaTime);
 
diff --git a/Engine/Source/Editor/StorageServerWidgets/Private/SZenServiceStatus.cpp b/Engine/Source/Editor/StorageServerWidgets/Private/SZenServiceStatus.cpp
--- a/Engine/Source/Editor/StorageServerWidgets/Private/SZenServiceStatus.cpp
+++ b/Engine/Source/Editor/StorageServerWidgets/Private/SZenServiceStatus.cpp
@@ -9,6 +9,7 @@
 #include "Styling/StyleColors.h"
 #include "Widgets/Images/SImage.h"
 #include "Widgets/Input/SButton.h"
+#include "Widgets/Input/SHyperlink.h"
 #include "Widgets/Layout/SGridPanel.h"
 #include "Widgets/SBoxPanel.h"
 #include "Widgets/Text/STextBlock.h"
@@ -79,6 +80,18 @@ FReply SZenServiceStatus::ExploreDataPath_OnClicked()
 	return FReply::Unhandled();
 }
 
+FReply SZenServiceStatus::ExploreExecutablePath_OnClicked()
+{
+	const FState& CurrentState = GetCurrentState();
+	if (CurrentState.bGotRunContext)
+	{
+		FString FullPath(FPaths::ConvertRelativePathToFull(FPaths::GetPath(CurrentState.RunContext.GetExecutable())));
+		FPlatformProcess::ExploreFolder(*FullPath);
+		return FReply::Handled();
+	}
+	return FReply::Unhandled();
+}
+
 const SZenServiceStatus::FState& SZenServiceStatus::GetCurrentState() const
 {
 	return State[ActiveStateIndex.load(std::memory_order_relaxed)];
@@ -178,6 +191,43 @@ TSharedRef<SWidget> SZenServiceStatus::GetGridPanel()
 
 	++Row;
 
+	Panel->AddSlot(0, Row)
+	[
+		SNew(STextBlock)
+		.Margin(FMargin(ColumnMargin, RowMargin))
+		.ColorAndOpacity(TitleColor)
+		.Font(TitleFont)
+		.Text(LOCTEXT("ServiceStatus_Browse", "Browse"))
+	];
+
+	Panel->AddSlot(1, Row)
+	[
+		SNew(SHyperlink)
+		.Padding(FMargin(ColumnMargin, RowMargin))
+		.UnderlineStyle(FAppStyle::Get(), "NoBorder")
+		.Text_Lambda([this]
+		{
+			const FState& CurrentState = GetCurrentState();
+			FText BrowseText = LOCTEXT("ServiceStatus_NoBrowseValue", "-");
+			if (CurrentState.bGotRunContext && CurrentState.bIsRunning && (CurrentState.LocalPort != 0))
+			{
+				BrowseText = FText::Format(LOCTEXT("ServiceStatus_Browse", "http://localhost:{0}/dashboard/"),
+					FText::AsNumber(CurrentState.LocalPort, &FNumberFormattingOptions::DefaultNoGrouping()));
+			}
+			return BrowseText;
+		})
+		.OnNavigate_Lambda([this]
+		{
+			const FState& CurrentState = GetCurrentState();
+			if (CurrentState.bGotRunContext && CurrentState.bIsRunning && (CurrentState.LocalPort != 0))
+			{
+				FPlatformProcess::LaunchURL(*FString::Printf(TEXT("http://localhost:%d/dashboard/"), CurrentState.LocalPort), nullptr, nullptr);
+			}
+		})
+	];
+
+	++Row;
+
 	Panel->AddSlot(0, Row)
 	[
 		SNew(STextBlock)
@@ -239,10 +289,13 @@ TSharedRef<SWidget> SZenServiceStatus::GetGridPanel()
 		SNew(SHorizontalBox)
 
 		+ SHorizontalBox::Slot()
+		.AutoWidth()
+		.MaxWidth(350)
 		.VAlign(VAlign_Bottom)
 		[
 			SNew(STextBlock)
 			.Margin(FMargin(ColumnMargin, RowMargin))
+			.OverflowPolicy(ETextOverflowPolicy::MiddleEllipsis)
 			.Text_Lambda([this]
 			{
 				const FState& CurrentState = GetCurrentState();
@@ -273,6 +326,59 @@ TSharedRef<SWidget> SZenServiceStatus::GetGridPanel()
 
 	++Row;
 
+	Panel->AddSlot(0, Row)
+	.VAlign(VAlign_Bottom)
+	[
+		SNew(STextBlock)
+		.Margin(FMargin(ColumnMargin, RowMargin))
+		.ColorAndOpacity(TitleColor)
+		.Font(TitleFont)
+		.Text(LOCTEXT("ServiceStatus_ExecutablePath", "Executable path"))
+	];
+
+	Panel->AddSlot(1, Row)
+	[
+
+		SNew(SHorizontalBox)
+
+		+ SHorizontalBox::Slot()
+		.AutoWidth()
+		//.MaxWidth(350)
+		.VAlign(VAlign_Bottom)
+		[
+			SNew(STextBlock)
+			.Margin(FMargin(ColumnMargin, RowMargin))
+			.OverflowPolicy(ETextOverflowPolicy::MiddleEllipsis)
+			.Text_Lambda([this]
+			{
+				const FState& CurrentState = GetCurrentState();
+				return CurrentState.bGotRunContext ? FText::FromString(FPaths::GetPath(CurrentState.RunContext.GetExecutable())) : LOCTEXT("ServiceStatus_NoExecutablePathValue", "-");
+			})
+		]
+
+		+ SHorizontalBox::Slot()
+		.AutoWidth()
+		.VAlign(VAlign_Bottom)
+		[
+			SNew(SButton)
+			.ButtonStyle(&FAppStyle::Get().GetWidgetStyle<FButtonStyle>("SimpleButton"))
+			.ToolTipText(LOCTEXT("ExploreExecutablePathButtonToolTip", "Explore the Executable Path"))
+			.IsEnabled_Lambda([this]
+			{
+				const FState& CurrentState = GetCurrentState();
+				return CurrentState.bGotRunContext;
+			})
+			.OnClicked(this, &SZenServiceStatus::ExploreExecutablePath_OnClicked)
+			[
+				SNew(SImage)
+				.Image(FAppStyle::Get().GetBrush("Zen.Icons.FolderExplore"))
+				.ColorAndOpacity(FSlateColor::UseForeground())
+			]
+		]
+	];
+
+	++Row;
+
 
 	return Panel;
 }
