Style changes and image button for the browse field of zen status panel.

diff --git a/Engine/Source/Editor/StorageServerWidgets/Private/SZenServiceStatus.cpp b/Engine/Source/Editor/StorageServerWidgets/Private/SZenServiceStatus.cpp
--- a/Engine/Source/Editor/StorageServerWidgets/Private/SZenServiceStatus.cpp
+++ b/Engine/Source/Editor/StorageServerWidgets/Private/SZenServiceStatus.cpp
@@ -202,28 +202,65 @@ TSharedRef<SWidget> SZenServiceStatus::GetGridPanel()
 
 	Panel->AddSlot(1, Row)
 	[
-		SNew(SHyperlink)
+		SNew(SHorizontalBox)
+
+		+ SHorizontalBox::Slot()
+		.AutoWidth()
+		.MaxWidth(350)
 		.Padding(FMargin(ColumnMargin, RowMargin))
-		.UnderlineStyle(FAppStyle::Get(), "NoBorder")
-		.Text_Lambda([this]
-		{
-			const FState& CurrentState = GetCurrentState();
-			FText BrowseText = LOCTEXT("ServiceStatus_NoBrowseValue", "-");
-			if (CurrentState.bGotRunContext && CurrentState.bIsRunning && (CurrentState.LocalPort != 0))
+		.VAlign(VAlign_Bottom)
+		[
+			SNew(SHyperlink)
+			.Style(FAppStyle::Get(), TEXT("NavigationHyperlink"))
+			.Text_Lambda([this]
 			{
-				BrowseText = FText::Format(LOCTEXT("ServiceStatus_BrowseLink", "http://localhost:{0}/dashboard/"),
-					FText::AsNumber(CurrentState.LocalPort, &FNumberFormattingOptions::DefaultNoGrouping()));
-			}
-			return BrowseText;
-		})
-		.OnNavigate_Lambda([this]
-		{
-			const FState& CurrentState = GetCurrentState();
-			if (CurrentState.bGotRunContext && CurrentState.bIsRunning && (CurrentState.LocalPort != 0))
+				const FState& CurrentState = GetCurrentState();
+				FText BrowseText = LOCTEXT("ServiceStatus_NoBrowseValue", "-");
+				if (CurrentState.bGotRunContext && CurrentState.bIsRunning && (CurrentState.LocalPort != 0))
+				{
+					BrowseText = FText::Format(LOCTEXT("ServiceStatus_BrowseLink", "http://localhost:{0}/dashboard/"),
+						FText::AsNumber(CurrentState.LocalPort, &FNumberFormattingOptions::DefaultNoGrouping()));
+				}
+				return BrowseText;
+			})
+			.OnNavigate_Lambda([this]
 			{
-				FPlatformProcess::LaunchURL(*FString::Printf(TEXT("http://localhost:%d/dashboard/"), CurrentState.LocalPort), nullptr, nullptr);
-			}
-		})
+				const FState& CurrentState = GetCurrentState();
+				if (CurrentState.bGotRunContext && CurrentState.bIsRunning && (CurrentState.LocalPort != 0))
+				{
+					FPlatformProcess::LaunchURL(*FString::Printf(TEXT("http://localhost:%d/dashboard/"), CurrentState.LocalPort), nullptr, nullptr);
+				}
+			})
+		]
+
+		+ SHorizontalBox::Slot()
+		.AutoWidth()
+		.VAlign(VAlign_Bottom)
+		[
+			SNew(SButton)
+			.ButtonStyle(&FAppStyle::Get().GetWidgetStyle<FButtonStyle>("SimpleButton"))
+			.ToolTipText(LOCTEXT("ServiceStatus_BrowseToolTip", "Browse the contents of zenserver in your web browser"))
+			.IsEnabled_Lambda([this]
+			{
+				const FState& CurrentState = GetCurrentState();
+				return CurrentState.bGotRunContext && CurrentState.bIsRunning && (CurrentState.LocalPort != 0);
+			})
+			.OnClicked_Lambda([this]
+			{
+				const FState& CurrentState = GetCurrentState();
+				if (CurrentState.bGotRunContext && CurrentState.bIsRunning && (CurrentState.LocalPort != 0))
+				{
+					FPlatformProcess::LaunchURL(*FString::Printf(TEXT("http://localhost:%d/dashboard/"), CurrentState.LocalPort), nullptr, nullptr);
+					return FReply::Handled();
+				}
+				return FReply::Unhandled();
+			})
+			[
+				SNew(SImage)
+				.Image(FAppStyle::Get().GetBrush("Zen.Icons.WebBrowser"))
+				.ColorAndOpacity(FSlateColor::UseForeground())
+			]
+		]
 	];
 
 	++Row;
diff --git a/Engine/Source/Programs/ZenDashboard/Private/ZenDashboardStyle.cpp b/Engine/Source/Programs/ZenDashboard/Private/ZenDashboardStyle.cpp
--- a/Engine/Source/Programs/ZenDashboard/Private/ZenDashboardStyle.cpp
+++ b/Engine/Source/Programs/ZenDashboard/Private/ZenDashboardStyle.cpp
@@ -51,11 +51,15 @@ TSharedRef< FSlateStyleSet > FZenDashboardStyle::Create()
 
 	Style.SetParentStyleName("CoreStyle");
 
+	const ISlateStyle* ParentStyle = FSlateStyleRegistry::FindSlateStyle("CoreStyle");
+	const FTextBlockStyle NormalText = ParentStyle->GetWidgetStyle<FTextBlockStyle>("NormalText");
+
 	const FVector2D Icon12x12(12.0f, 12.0f);
 	const FVector2D Icon16x16(16.0f, 16.0f);
 
 	Style.Set( "AppIcon", new IMAGE_BRUSH( "Icons/DefaultAppIcon", FVector2D(20,20)) );
 
+	Style.Set("Zen.Icons.WebBrowser", new EDITOR_IMAGE_BRUSH_SVG("Starship/Common/WebBrowser", Icon12x12));
 	Style.Set("Zen.Icons.FolderExplore", new EDITOR_IMAGE_BRUSH_SVG("Starship/Common/ContentBrowser", Icon12x12));
 	Style.Set("Zen.Icons.Clipboard", new EDITOR_IMAGE_BRUSH_SVG("Starship/Common/Clipboard", Icon16x16));
 
@@ -79,6 +83,25 @@ TSharedRef< FSlateStyleSet > FZenDashboardStyle::Create()
 		.SetColorAndOpacity(FSlateColor::UseSubduedForeground())
 	);
 
+	{
+		// Navigation defaults
+		const FLinearColor NavHyperlinkColor(0.03847f, 0.33446f, 1.0f);
+		const FTextBlockStyle NavigationHyperlinkText = FTextBlockStyle(NormalText)
+			.SetFont(DEFAULT_FONT("Regular", 10))
+			.SetColorAndOpacity(NavHyperlinkColor);
+
+		const FButtonStyle NavigationHyperlinkButton = FButtonStyle()
+			.SetNormal(BORDER_BRUSH("Old/HyperlinkDotted", FMargin(0, 0, 0, 3 / 16.0f), NavHyperlinkColor))
+			.SetPressed(FSlateNoResource())
+			.SetHovered(BORDER_BRUSH("Old/HyperlinkUnderline", FMargin(0, 0, 0, 3 / 16.0f), NavHyperlinkColor));
+
+		FHyperlinkStyle NavigationHyperlink = FHyperlinkStyle()
+			.SetUnderlineStyle(NavigationHyperlinkButton)
+			.SetTextStyle(NavigationHyperlinkText)
+			.SetPadding(FMargin(0.0f));
+
+		Style.Set("NavigationHyperlink", NavigationHyperlink);
+	}
 
 	FSlateBrush* GenericWhiteBox = new IMAGE_BRUSH( "Old/White", Icon16x16 );
 
