Zen dashboard workspace support

#rnx

diff --git a/Engine/Source/Developer/Zen/Private/ZenServerInterface.cpp b/Engine/Source/Developer/Zen/Private/ZenServerInterface.cpp
--- a/Engine/Source/Developer/Zen/Private/ZenServerInterface.cpp
+++ b/Engine/Source/Developer/Zen/Private/ZenServerInterface.cpp
@@ -29,6 +29,7 @@
 #include "Misc/PathViews.h"
 #include "Misc/ScopedSlowTask.h"
 #include "Misc/ScopeExit.h"
+#include "Serialization/MemoryReader.h"
 #include "Serialization/CompactBinarySerialization.h"
 #include "Serialization/JsonReader.h"
 #include "Serialization/JsonSerializer.h"
@@ -2892,6 +2893,97 @@ FZenServiceInstance::GatherAnalytics(TArray<FAnalyticsEventAttribute>& Attribute
 	return true;
 }
 
+bool
+FZenServiceInstance::GetWorkspaces(FZenWorkspaces& Workspaces) const
+{
+	TStringBuilder<128> ZenDomain;
+	ZenDomain << HostName << TEXT(":") << Port;
+
+#if WITH_EDITOR
+	EAsyncExecution ThreadPool = EAsyncExecution::LargeThreadPool;
+#else
+	EAsyncExecution ThreadPool = EAsyncExecution::ThreadPool;
+#endif
+
+	WorkspacesRequest = Async(ThreadPool, [ZenDomain = FString(ZenDomain)]
+	{
+		UE::Zen::FZenHttpRequest Request(ZenDomain, false);
+
+		FZenHttpRequest::Result Result = Request.PerformBlockingDownload(TEXTVIEW("/ws"), nullptr, Zen::EContentType::CbObject);
+
+		FZenWorkspaces Workspaces;
+
+		if (Result == Zen::FZenHttpRequest::Result::Success && Request.GetResponseCode() == 200)
+		{
+			FMemoryReaderView Reader(Request.GetResponseBuffer());
+			FCbObject ResponseObj = LoadCompactBinary(Reader).AsObject();
+
+			FCbArrayView WorkspacesArray = ResponseObj["workspaces"].AsArrayView();
+			for (FCbFieldView WorkspaceField : WorkspacesArray)
+			{
+				FCbObjectView WorkspaceObject = WorkspaceField.AsObjectView();
+
+				FZenWorkspaces::Workspace Workspace = {
+					.Id = *WriteToString<64>(WorkspaceObject["id"].AsObjectId()),
+					.BaseDir = FString(WorkspaceObject["root_path"].AsString()),
+					.bDynamicShare = WorkspaceObject["allow_share_creation_from_http"].AsBool()
+				};
+
+				if (!Workspace.Id.IsEmpty())
+				{
+					FCbArrayView WorkspaceShares = WorkspaceObject["shares"].AsArrayView();
+
+					UE_LOG(LogZenServiceInstance, Warning, TEXT("Number of shares %d"), WorkspaceShares.Num());
+
+					for (int32 Idx = 0; Idx < WorkspaceShares.Num(); ++Idx)
+					{
+						UE_LOG(LogZenServiceInstance, Warning, TEXT("Iterating through share %d"), Idx);
+					}
+
+					for (FCbFieldView ShareField : WorkspaceShares)
+					{
+						FCbObjectView ShareObject = ShareField.AsObjectView();
+
+						FZenWorkspaces::Share Share = {
+							.Id = *WriteToString<64>(ShareObject["id"].AsObjectId()),
+							.Dir = FString(ShareObject["share_path"].AsString()),
+							.Alias = FString(ShareObject["alias"].AsString())
+						};
+
+						if (!Share.Id.IsEmpty())
+						{
+							Workspace.WorkspaceShares.Add(Share);
+						}
+					}
+
+					Workspaces.ZenWorkspaces.Add(Workspace);
+				}
+			}
+
+			Workspaces.bIsValid = true;
+		}
+
+		return Workspaces;
+	});
+
+	WorkspacesRequest.Wait();
+
+	Workspaces = WorkspacesRequest.Get();
+	return Workspaces.bIsValid;
+}
+
+int32
+FZenServiceInstance::GetWorkspaceCount() const
+{
+	FZenWorkspaces Workspaces;
+	if (!GetWorkspaces(Workspaces))
+	{
+		return 0;
+	}
+
+	return Workspaces.ZenWorkspaces.Num();
+}
+
 #endif // UE_WITH_ZEN
 
 }
diff --git a/Engine/Source/Developer/Zen/Public/Experimental/ZenServerInterface.h b/Engine/Source/Developer/Zen/Public/Experimental/ZenServerInterface.h
--- a/Engine/Source/Developer/Zen/Public/Experimental/ZenServerInterface.h
+++ b/Engine/Source/Developer/Zen/Public/Experimental/ZenServerInterface.h
@@ -71,6 +71,69 @@ private:
 	bool TryApplyAutoLaunchOverride();
 };
 
+struct FZenWorkspaces
+{
+	struct Share
+	{
+		FString Id;
+		FString Dir;
+		FString Alias;
+	};
+
+	struct Workspace
+	{
+		FString Id;
+		FString BaseDir;
+		bool bDynamicShare;
+
+		TArray<Share> WorkspaceShares;
+	};
+
+	TArray<Workspace> ZenWorkspaces;
+	bool bIsValid = false;
+
+	bool operator==(const FZenWorkspaces& Other) const
+	{
+		if (ZenWorkspaces.Num() != Other.ZenWorkspaces.Num())
+		{
+			return false;
+		}
+
+		for (int32 Idx = 0; Idx < ZenWorkspaces.Num(); ++Idx)
+		{
+			const Workspace& CurrentWorkspace = ZenWorkspaces[Idx];
+			const Workspace& OtherWorkspace = Other.ZenWorkspaces[Idx];
+
+			if (CurrentWorkspace.WorkspaceShares.Num() != OtherWorkspace.WorkspaceShares.Num())
+			{
+				return false;
+			}
+
+			if (CurrentWorkspace.Id != OtherWorkspace.Id ||
+				CurrentWorkspace.BaseDir != OtherWorkspace.BaseDir ||
+				CurrentWorkspace.bDynamicShare != OtherWorkspace.bDynamicShare)
+			{
+				return false;
+			}
+
+			for (int32 ShareIdx = 0; ShareIdx < CurrentWorkspace.WorkspaceShares.Num(); ++ShareIdx)
+			{
+				const Share& CurrentShare = CurrentWorkspace.WorkspaceShares[ShareIdx];
+				const Share& OtherShare = OtherWorkspace.WorkspaceShares[ShareIdx];
+
+				if (CurrentShare.Id != OtherShare.Id ||
+					CurrentShare.Dir != OtherShare.Dir ||
+					CurrentShare.Alias != OtherShare.Alias)
+				{
+					return false;
+				}
+			}
+		}
+
+		return true;
+	}
+};
+
 }
 
 #if UE_WITH_ZEN
@@ -202,6 +265,9 @@ public:
 	UE_API bool GetProjectStats(FZenProjectStats& Stats);
 	UE_API bool GatherAnalytics(TArray<FAnalyticsEventAttribute>& Attributes);
 
+	UE_API int32 GetWorkspaceCount() const;
+	UE_API bool GetWorkspaces(FZenWorkspaces& Workspaces) const;
+
 	UE_API bool GetGCStatus(FGCStatus& Status);
 	UE_API bool RequestGC(const bool* OverrideCollectSmallObjects = nullptr, const uint32* OverrideMaxCacheDuration = nullptr);
 
@@ -230,6 +296,8 @@ private:
 	mutable FMutex LastProjectStatsMutex;
 	mutable FZenProjectStats LastProjectStats;
 
+	mutable TFuture<FZenWorkspaces> WorkspacesRequest;
+
 	FServiceSettings Settings;
 	FString URL;
 	FString HostName;
diff --git a/Engine/Source/Editor/StorageServerWidgets/Internal/SZenWorkspaces.h b/Engine/Source/Editor/StorageServerWidgets/Internal/SZenWorkspaces.h
new file mode 100644
--- /dev/null
+++ b/Engine/Source/Editor/StorageServerWidgets/Internal/SZenWorkspaces.h
@@ -0,0 +1,56 @@
+// Copyright Epic Games, Inc. All Rights Reserved.
+
+#pragma once
+
+#include "Experimental/ZenServerInterface.h"
+#include "ZenServiceInstanceManager.h"
+#include "Templates/SharedPointer.h"
+#include "Widgets/DeclarativeSyntaxSupport.h"
+#include "Widgets/SBoxPanel.h"
+#include "Widgets/SCompoundWidget.h"
+#include "Widgets/SWidget.h"
+#include "Widgets/Layout/SWidgetSwitcher.h"
+#include "Widgets/Views/SListView.h"
+#include "Widgets/Layout/SScrollBox.h"
+#include "Widgets/Layout/SExpandableArea.h"
+
+class SZenWorkspaces : public SCompoundWidget
+{
+	SLATE_BEGIN_ARGS(SZenWorkspaces)
+		: _ZenServiceInstance(nullptr)
+		{
+		}
+
+		SLATE_ATTRIBUTE(TSharedPtr<UE::Zen::FZenServiceInstance>, ZenServiceInstance);
+		SLATE_ATTRIBUTE(float, UpdateFrequency);
+
+	SLATE_END_ARGS()
+
+	void Construct(const FArguments& InArgs);
+
+	int32 GetWorkspaceCount() const { return Workspaces.ZenWorkspaces.Num(); }
+	void UpdateWorkspaces(float InDeltaTime, bool bForce = false);
+
+private:
+	TSharedRef<SWidget> GetWorkspaceList();
+
+	FReply ZenWorkspacesArea_Toggle() const;
+
+	FReply OnBackToMainWidget();
+	FReply OnCreateNewWorkspaceClicked();
+
+	const FSlateBrush* ZenWorkspacesArea_Icon() const;
+
+private:
+	TAttribute<TSharedPtr<UE::Zen::FZenServiceInstance>> ZenServiceInstance;
+	float UpdateFrequency = 0.f;
+	TSharedPtr<SVerticalBox> WorkspaceArea;
+	TSharedPtr<SWidgetSwitcher> WidgetSwitcher;
+	SVerticalBox::FSlot* GridSlot = nullptr;
+	TSharedPtr<SExpandableArea> Expandable;
+
+	TArray<TSharedPtr<UE::Zen::FZenWorkspaces>> WorkspaceModels;
+	TSharedPtr<SListView<TSharedPtr<UE::Zen::FZenWorkspaces>>> WorkspacesListView;
+
+	UE::Zen::FZenWorkspaces Workspaces;
+};
diff --git a/Engine/Source/Editor/StorageServerWidgets/Private/SZenWorkspaces.cpp b/Engine/Source/Editor/StorageServerWidgets/Private/SZenWorkspaces.cpp
new file mode 100644
--- /dev/null
+++ b/Engine/Source/Editor/StorageServerWidgets/Private/SZenWorkspaces.cpp
@@ -0,0 +1,726 @@
+// Copyright Epic Games, Inc. All Rights Reserved.
+
+#include "SZenWorkspaces.h"
+
+#include "SZenServiceStatus.h"
+#include "Experimental/ZenServerInterface.h"
+#include "Internationalization/FastDecimalFormat.h"
+#include "Math/BasicMathExpressionEvaluator.h"
+#include "Math/UnitConversion.h"
+#include "Misc/ExpressionParser.h"
+#include "Styling/StyleColors.h"
+#include "Widgets/Images/SImage.h"
+#include "Widgets/Input/SButton.h"
+#include "Widgets/Layout/SGridPanel.h"
+#include "Widgets/Layout/SSeparator.h"
+#include "Widgets/SBoxPanel.h"
+#include "Widgets/Text/STextBlock.h"
+#include "Widgets/Input/SEditableTextBox.h"
+#include "Framework/Application/SlateApplication.h"
+#include "DesktopPlatformModule.h"
+#include "Misc/MonitoredProcess.h"
+
+//#include "ZenDashboardStyle.h"
+
+#define LOCTEXT_NAMESPACE "ZenDashboard"
+
+DEFINE_LOG_CATEGORY_STATIC(LogZenDashboardWorkpaces, Log, All);
+
+namespace
+{
+	bool RunZenExe(FString Cmd)
+	{
+		FString ZenExePath = UE::Zen::GetLocalInstallUtilityPath();
+
+		FMonitoredProcess ZenExeProcess(ZenExePath, Cmd, true);
+		bool bLaunched = ZenExeProcess.Launch();
+		if (!bLaunched)
+		{
+			UE_LOG(LogZenDashboardWorkpaces, Warning, TEXT("Failed to launch zen utility: '%s'."), *ZenExePath);
+			return false;
+		}
+
+		const uint64 StartTime = FPlatformTime::Cycles64();
+		while (ZenExeProcess.Update())
+		{
+			double Duration = FPlatformTime::ToSeconds64(FPlatformTime::Cycles64() - StartTime);
+			if (Duration > 10.f)
+			{
+				ZenExeProcess.Cancel(true);
+				UE_LOG(LogZenDashboardWorkpaces, Warning, TEXT("Cancelled launch of zen utility: '%s %s' due to timeout."), *ZenExePath, *Cmd);
+
+				while (ZenExeProcess.Update())
+				{
+					Duration = FPlatformTime::ToSeconds64(FPlatformTime::Cycles64() - StartTime);
+					if (Duration > 15.f)
+					{
+						UE_LOG(LogZenDashboardWorkpaces, Warning, TEXT("Cancelled launch of zen utility: '%s'. Failed waiting for termination."), *ZenExePath);
+						break;
+					}
+
+					FPlatformProcess::Sleep(0.2f);
+				}
+
+				FString OutputString = ZenExeProcess.GetFullOutputWithoutDelegate();
+				UE_LOG(LogZenDashboardWorkpaces, Warning, TEXT("Launch of zen utility: '%s' failed. Output: '%s'"), *ZenExePath, *OutputString);
+				return false;
+			}
+
+			FPlatformProcess::Sleep(0.1f);
+		}
+
+		return true;
+	}
+}
+
+class SZenNewWorkspace : public SCompoundWidget
+{
+	SLATE_BEGIN_ARGS(SZenNewWorkspace)
+	{}
+		SLATE_EVENT(FOnClicked, OnCloseClicked);
+	SLATE_END_ARGS()
+
+	void Construct(const FArguments& InArgs)
+	{
+		OnCloseClicked = InArgs._OnCloseClicked;
+
+		ChildSlot
+		[
+			SNew(SVerticalBox)
+
+			+SVerticalBox::Slot()
+			.AutoHeight()
+			.HAlign(HAlign_Left)
+			[
+				SNew(SButton)
+				.OnClicked(this, &SZenNewWorkspace::OnGoBack)
+				[
+					SNew(SImage)
+					.Image(FAppStyle::Get().GetBrush("Icons.Back"))
+				]
+			]
+
+			+ SVerticalBox::Slot()
+			.Padding(0.f, 4.f)
+			.AutoHeight()
+			[
+				SNew(SSeparator)
+			]
+
+			+SVerticalBox::Slot()
+			.AutoHeight()
+			[
+				SNew(SHorizontalBox)
+
+				+ SHorizontalBox::Slot()
+				.Padding(2.f)
+				.VAlign(VAlign_Center)
+				.AutoWidth()
+				[
+					SNew(STextBlock)
+					.Text(LOCTEXT("NewWorkspacePath", "Workspace Path:"))
+				]
+
+				+ SHorizontalBox::Slot()
+				.FillWidth(1)
+				.VAlign(VAlign_Center)
+				.Padding(0, 0, 0, 3)
+				[
+					SNew(SEditableTextBox)
+					.Text(this, &SZenNewWorkspace::GetCurrentPath)
+
+				]
+
+				+ SHorizontalBox::Slot()
+				.VAlign(VAlign_Center)
+				.AutoWidth()
+				[
+					SNew(SButton)
+					.ContentPadding(FMargin(6.f, 2.f))
+					.OnClicked(this, &SZenNewWorkspace::OnBrowseClick)
+					[
+						SNew(SImage)
+						.Image(FAppStyle::Get().GetBrush("Icons.FolderOpen"))
+					]
+				]
+			]
+
+			+ SVerticalBox::Slot()
+			.Padding(0.f, 4.f)
+			.AutoHeight()
+			[
+				SNew(SSeparator)
+			]
+
+			+ SVerticalBox::Slot()
+			.Padding(0.f, 4.f)
+			.AutoHeight()
+			.HAlign(HAlign_Right)
+			[
+				SNew(SButton)
+				.ContentPadding(FMargin(6.f, 2.f))
+				.OnClicked(this, &SZenNewWorkspace::OnGoBack)
+				[
+					SNew(SImage)
+					.Image(FAppStyle::Get().GetBrush("Icons.Check"))
+					.ColorAndOpacity(FStyleColors::AccentGreen)
+				]
+			]
+		];
+	}
+
+private:
+	FText GetCurrentPath() const
+	{
+		return FText::GetEmpty();
+	}
+
+	FReply OnGoBack() const
+	{
+		if (OnCloseClicked.IsBound())
+		{
+			OnCloseClicked.Execute();
+		}
+		return FReply::Handled();
+	}
+
+	FReply OnBrowseClick()
+	{
+		IDesktopPlatform* DesktopPlatform = FDesktopPlatformModule::Get();
+		if (DesktopPlatform)
+		{
+			TSharedPtr<SWindow> ParentWindow = FSlateApplication::Get().FindWidgetWindow(AsShared());
+			void* ParentWindowHandle = (ParentWindow.IsValid() && ParentWindow->GetNativeWindow().IsValid()) ? ParentWindow->GetNativeWindow()->GetOSWindowHandle() : nullptr;
+
+			FString FolderName;
+			const bool bFolderSelected = DesktopPlatform->OpenDirectoryDialog(
+				ParentWindowHandle,
+				LOCTEXT("FolderDialogTitle", "Choose a directory").ToString(),
+				"",
+				FolderName
+			);
+
+			if (bFolderSelected)
+			{
+				if (!FolderName.EndsWith(TEXT("/")))
+				{
+					FolderName += TEXT("/");
+				}
+
+				SelectedDir = FolderName;
+			}
+		}
+
+		return FReply::Handled();
+	}
+
+private:
+	FOnClicked OnCloseClicked;
+	FString SelectedDir;
+};
+
+class SZenWorkspaceShareRow : public SCompoundWidget
+{
+	SLATE_BEGIN_ARGS(SZenWorkspaceShareRow)
+		{}
+		SLATE_ATTRIBUTE(UE::Zen::FZenWorkspaces::Workspace, WorkspaceInfo);
+		SLATE_ATTRIBUTE(UE::Zen::FZenWorkspaces::Share, ShareInfo);
+	SLATE_END_ARGS()
+
+	void Construct(const FArguments& InArgs)
+	{
+		WorkspaceInfo = InArgs._WorkspaceInfo.Get();
+		ShareInfo = InArgs._ShareInfo.Get();
+
+		const float RowMargin = 10.f;
+		const float ColumnMargin = 0.f;
+		const FSlateFontInfo TitleFont = FCoreStyle::GetDefaultFontStyle("Bold", 10);
+
+		this->ChildSlot
+		[
+			SNew(SBorder)
+			//.Padding(10.f, 10.f)
+			.BorderImage(FAppStyle::GetBrush("ToolPanel.GroupBorder"))
+			[
+				SNew(SHorizontalBox)
+
+				+SHorizontalBox::Slot()
+				.FillWidth(2)
+				[
+					SNew(SGridPanel)
+			
+					+ SGridPanel::Slot(0, 0)
+					[
+						SNew(SBox)
+						[
+							SNew(STextBlock)
+							.Font(TitleFont)
+							.Margin(FMargin(RowMargin, ColumnMargin))
+							.ColorAndOpacity(FStyleColors::AccentWhite)
+							.Text(LOCTEXT("SharesId", "ShareId"))
+						]
+					]
+
+					+ SGridPanel::Slot(1, 0)
+					[
+						SNew(STextBlock)
+						.Margin(FMargin(RowMargin, ColumnMargin))
+						.Text_Lambda([this]()
+						{
+							return FText::FromString(ShareInfo.Id);
+						})
+					]
+
+					+ SGridPanel::Slot(0, 1)
+					[
+						SNew(STextBlock)
+						.Margin(FMargin(RowMargin, ColumnMargin))
+						.Font(TitleFont)
+						.ColorAndOpacity(FStyleColors::AccentWhite)
+						.Text(LOCTEXT("SharePath", "Path"))
+					]
+
+					+ SGridPanel::Slot(1, 1)
+					[
+						SNew(STextBlock)
+						.Margin(FMargin(RowMargin, ColumnMargin))
+						.Text_Lambda([this]()
+						{
+							return FText::FromString(ShareInfo.Dir);
+						})
+					]
+
+					+ SGridPanel::Slot(0, 2)
+					[
+						SNew(STextBlock)
+						.Margin(FMargin(RowMargin, ColumnMargin))
+						.Font(TitleFont)
+						.ColorAndOpacity(FStyleColors::AccentWhite)
+						.Text(LOCTEXT("SharePath", "Alias"))
+					]
+
+					+ SGridPanel::Slot(1, 2)
+					[
+						SNew(STextBlock)
+						.Margin(FMargin(RowMargin, ColumnMargin))
+						.Text_Lambda([this]()
+						{
+							return FText::FromString(ShareInfo.Alias);
+						})
+					]
+				]
+
+				+ SHorizontalBox::Slot()
+				.AutoWidth()
+				[
+					SNew(SBox)
+					[
+						SNew(SButton)
+						.VAlign(VAlign_Center)
+						.HAlign(HAlign_Center)
+						.ButtonStyle(&FAppStyle::Get().GetWidgetStyle<FButtonStyle>("Button"))
+						.OnClicked(this, &SZenWorkspaceShareRow::OnDeleteClicked)
+						[
+							SNew(SImage)
+							.Image(FAppStyle::Get().GetBrush("Icons.XCircle"))
+							.ColorAndOpacity(FStyleColors::AccentRed)
+						]
+					]
+				]
+			]
+		];
+	}
+
+	FReply OnDeleteClicked()
+	{
+		TStringBuilder<128> ShareRemoveCommand;
+		ShareRemoveCommand << "workspace-share remove " << WorkspaceInfo.Id << " " << ShareInfo.Id;
+
+		if (RunZenExe(ShareRemoveCommand.ToString()))
+		{
+			SetVisibility(EVisibility::Collapsed);
+		}
+
+		return FReply::Handled();
+	}
+
+private:
+	UE::Zen::FZenWorkspaces::Workspace WorkspaceInfo;
+	UE::Zen::FZenWorkspaces::Share ShareInfo;
+};
+
+class SZenWorkspaceRow : public SCompoundWidget
+{
+	SLATE_BEGIN_ARGS(SZenWorkspaceRow)
+		{}
+		SLATE_ATTRIBUTE(UE::Zen::FZenWorkspaces::Workspace, WorkspaceInfo);
+		SLATE_ATTRIBUTE(int32, IndexNumber)
+	SLATE_END_ARGS()
+
+	void Construct(const FArguments& InArgs)
+	{
+		WorkspaceInfo = InArgs._WorkspaceInfo.Get();
+		IndexNumber = InArgs._IndexNumber.Get();
+
+		const float RowMargin = 0.f;
+		const float ColumnMargin = 10.f;
+		const FSlateFontInfo TitleFont = FCoreStyle::GetDefaultFontStyle("Bold", 10);
+
+		this->ChildSlot
+		[
+			SNew(SBorder)
+			.BorderImage(FAppStyle::GetBrush("Brushes.Secondary"))
+			[
+				SNew(SHorizontalBox)
+
+				+ SHorizontalBox::Slot()
+				.AutoWidth()
+				.FillWidth(1.0)
+				[
+					SNew(SVerticalBox)
+
+					+ SVerticalBox::Slot()
+					[
+						SNew(SHorizontalBox)
+
+						+ SHorizontalBox::Slot()
+						[
+							SNew(SGridPanel)
+
+							// base path row
+							+ SGridPanel::Slot(0, 0)
+							[
+								SNew(STextBlock)
+								.Margin(FMargin(ColumnMargin, RowMargin))
+								.ColorAndOpacity(FStyleColors::AccentWhite)
+								.Font(TitleFont)
+								.Text(LOCTEXT("ZenWorkspaces_Path", "Root dir"))
+							]
+							+ SGridPanel::Slot(1, 0)
+							[
+								SNew(STextBlock)
+								.Margin(FMargin(ColumnMargin, RowMargin))
+								.Text_Lambda([this]
+								{
+									return FText::FromString(WorkspaceInfo.BaseDir);
+								})
+							]
+
+							// id row
+							+ SGridPanel::Slot(0, 1)
+							[
+								SNew(STextBlock)
+								.Margin(FMargin(ColumnMargin, RowMargin))
+								.ColorAndOpacity(FStyleColors::AccentWhite)
+								.Font(TitleFont)
+								.Text(LOCTEXT("ZenWorkspaces_Id", "Id"))
+							]
+							+ SGridPanel::Slot(1, 1)
+							[
+								SNew(STextBlock)
+								.Margin(FMargin(ColumnMargin, RowMargin))
+								.Text_Lambda([this]
+								{
+									return FText::FromString(WorkspaceInfo.Id);
+								})
+							]
+
+							+ SGridPanel::Slot(0, 2)
+							[
+								SNew(STextBlock)
+								.Margin(FMargin(ColumnMargin, RowMargin))
+								.ColorAndOpacity(FStyleColors::AccentWhite)
+								.Font(TitleFont)
+								.Text(LOCTEXT("ZenWorkspaces_Dynamic", "Dynamic shares"))
+								.ToolTipText(LOCTEXT("WorkspacesDynamic_Tooltip", "Does this workspace allow for creating shares using the /ws http endpoint"))
+							]
+							+ SGridPanel::Slot(1, 2)
+							[
+								SNew(STextBlock)
+								.Margin(FMargin(ColumnMargin, RowMargin))
+								.Text_Lambda([this]
+								{
+									FText DynamicStatus = LOCTEXT("WorkspaceStatus_Disabled", "Disabled");
+									if (WorkspaceInfo.bDynamicShare)
+									{
+										DynamicStatus = LOCTEXT("WorkspaceStatus_Enabled", "Enabled");
+									}
+
+									return DynamicStatus;
+								})
+							]
+						]
+
+						+ SHorizontalBox::Slot()
+						.AutoWidth()
+						.Padding(0.f, 0.f, 10.f, 0.f)
+						[
+							SNew(SBox)
+							[
+								SNew(SButton)
+								.VAlign(VAlign_Center)
+								.HAlign(HAlign_Center)
+								.OnClicked(this, &SZenWorkspaceRow::OnDeleteClicked)
+								.ButtonStyle(&FAppStyle::Get().GetWidgetStyle<FButtonStyle>("Button"))
+								[
+									SNew(SImage)
+									.Image(FAppStyle::Get().GetBrush("GenericCommands.Delete"))
+									.ColorAndOpacity(FStyleColors::AccentRed)
+								]
+							]
+						]
+					]
+
+					+ SVerticalBox::Slot()
+					.AutoHeight()
+					.Padding(8.f, 8.f)
+					[
+						ConstructShareList()
+					]
+				]
+			]
+		];
+	}
+
+	FReply OnDeleteClicked()
+	{
+		TStringBuilder<128> WorkspaceRemoveCommand;
+		WorkspaceRemoveCommand << "workspace remove " << WorkspaceInfo.Id;
+
+		for (int Idx = 0; Idx < WorkspaceInfo.WorkspaceShares.Num(); ++Idx)
+		{
+			TStringBuilder<128> ShareRemoveCommand;
+			ShareRemoveCommand << "workspace-share remove " << WorkspaceInfo.Id << " " << WorkspaceInfo.WorkspaceShares[Idx].Id;
+
+			RunZenExe(ShareRemoveCommand.ToString());
+		}
+
+		if (RunZenExe(WorkspaceRemoveCommand.ToString()))
+		{
+			SetVisibility(EVisibility::Collapsed);
+		}
+
+		return FReply::Handled();
+	}
+
+private:
+	TSharedRef<SWidget> ConstructShareList()
+	{
+		const float RowMargin = 0.f;
+		const float ColumnMargin = 10.f;
+		const FSlateFontInfo TitleFont = FCoreStyle::GetDefaultFontStyle("Bold", 10);
+
+		TSharedRef<SVerticalBox> Widget = SNew(SVerticalBox);
+
+		if (WorkspaceInfo.WorkspaceShares.IsEmpty())
+		{
+			return Widget;
+		}
+
+		for (int32 Idx = 0; Idx < WorkspaceInfo.WorkspaceShares.Num(); ++Idx)
+		{
+			Widget->AddSlot()
+			.AutoHeight()
+			.Padding(0.f, 1.f)
+			[
+				SNew(SZenWorkspaceShareRow)
+				.WorkspaceInfo(WorkspaceInfo)
+				.ShareInfo(WorkspaceInfo.WorkspaceShares[Idx])
+			];
+		}
+
+		TSharedRef<SBorder> Result = SNew(SBorder)
+		.BorderImage(FAppStyle::GetBrush("Brushes.Header"))
+		[
+			SNew(SVerticalBox)
+
+			+ SVerticalBox::Slot()
+			.AutoHeight()
+			.HAlign(HAlign_Center)
+			[
+				SNew(STextBlock)
+					.Font(TitleFont)
+					.Margin(FMargin(RowMargin, ColumnMargin))
+					.ColorAndOpacity(FStyleColors::AccentWhite)
+					.Text(LOCTEXT("Shares", "Shared folders:"))
+			]
+
+			+SVerticalBox::Slot()
+			[
+				Widget
+			]
+		];
+
+		return Result;
+	}
+
+private:
+	UE::Zen::FZenWorkspaces::Workspace WorkspaceInfo;
+	int32 IndexNumber;
+	bool bVisible = true;
+};
+
+void SZenWorkspaces::Construct(const FArguments& InArgs)
+{
+	ZenServiceInstance = InArgs._ZenServiceInstance;
+	UpdateFrequency = InArgs._UpdateFrequency.Get();
+
+	const FSlateColor TitleColor = FStyleColors::AccentWhite;
+
+	this->ChildSlot
+		[
+			SNew(SVerticalBox)
+
+			+ SVerticalBox::Slot()
+			.AutoHeight()
+			.Padding(0.f, 0.f)
+			.HAlign(HAlign_Fill)
+			[
+				SNew(SButton)
+				.ButtonStyle(&FAppStyle::Get().GetWidgetStyle<FButtonStyle>("SimpleButton"))
+				.ToolTipText(LOCTEXT("ZenWorkspacesExpand", "Workspaces"))
+				.OnClicked(this, &SZenWorkspaces::ZenWorkspacesArea_Toggle)
+				[
+					SNew(SHorizontalBox)
+
+					+SHorizontalBox::Slot()
+					.AutoWidth()
+					.VAlign(VAlign_Center)
+					[
+						SNew(SImage)
+						.Image(this, &SZenWorkspaces::ZenWorkspacesArea_Icon)
+						.ColorAndOpacity(FStyleColors::White)
+					]
+
+					+SHorizontalBox::Slot()
+					.VAlign(VAlign_Center)
+					.Padding(4.0f, 0.f)
+					[
+						SNew(STextBlock)
+						.ColorAndOpacity(FStyleColors::White)
+						.Text(LOCTEXT("ZenWorkspaces", "Workspaces"))
+					]
+				]
+			]
+
+			+ SVerticalBox::Slot()
+			.MaxHeight(400)
+			.Padding(0.f, 8.f)
+			[
+				SAssignNew(WidgetSwitcher, SWidgetSwitcher)
+				.Visibility(EVisibility::Collapsed)
+
+				+ SWidgetSwitcher::Slot()
+				[
+					SAssignNew(WorkspaceArea, SVerticalBox)
+
+					+ SVerticalBox::Slot()
+					.Expose(GridSlot)
+					[
+						GetWorkspaceList()
+					]
+				]
+
+				+ SWidgetSwitcher::Slot()
+				[
+					SNew(SZenNewWorkspace)
+					.OnCloseClicked(this, &SZenWorkspaces::OnBackToMainWidget)
+				]
+			]
+		];
+
+	UpdateWorkspaces(0, true);
+}
+
+TSharedRef<SWidget> SZenWorkspaces::GetWorkspaceList()
+{
+	TSharedPtr<SScrollBox> WorkspaceItems = SNew(SScrollBox);
+
+	const float RowMargin = 0.f;
+	const float ColumnMargin = 10.f;
+	const FSlateFontInfo TitleFont = FCoreStyle::GetDefaultFontStyle("Bold", 10);
+
+	int32 Idx = 1;
+	for (UE::Zen::FZenWorkspaces::Workspace& Workspace : Workspaces.ZenWorkspaces)
+	{
+		WorkspaceItems->AddSlot()
+		.Padding(16.f, 4.f)
+		.VAlign(VAlign_Center)
+		[
+			SNew(SZenWorkspaceRow)
+			.WorkspaceInfo(Workspace)
+			.IndexNumber(Idx++)
+		];
+	}
+
+	return WorkspaceItems.ToSharedRef();
+}
+
+FReply SZenWorkspaces::ZenWorkspacesArea_Toggle() const
+{
+	if (WidgetSwitcher)
+	{
+		if (WidgetSwitcher->GetVisibility() == EVisibility::Visible)
+		{
+			WidgetSwitcher->SetVisibility(EVisibility::Collapsed);
+		}
+		else
+		{
+			WidgetSwitcher->SetVisibility(EVisibility::Visible);
+		}
+	}
+
+	return FReply::Handled();
+}
+
+FReply SZenWorkspaces::OnBackToMainWidget()
+{
+	WidgetSwitcher->SetActiveWidgetIndex(0);
+	return FReply::Handled();
+}
+
+FReply SZenWorkspaces::OnCreateNewWorkspaceClicked()
+{
+	WidgetSwitcher->SetActiveWidgetIndex(1);
+	return FReply::Handled();
+}
+
+const FSlateBrush* SZenWorkspaces::ZenWorkspacesArea_Icon() const
+{
+	return (WidgetSwitcher && WidgetSwitcher->GetVisibility() == EVisibility::Visible) ?
+		FAppStyle::Get().GetBrush("Icons.ChevronDown") :
+		FAppStyle::Get().GetBrush("Icons.ChevronRight");
+}
+
+void SZenWorkspaces::UpdateWorkspaces(float InDeltaTime, bool bForce)
+{
+	static double AccumulatedDeltaTime = 0.f;
+
+	AccumulatedDeltaTime += InDeltaTime;
+	if (AccumulatedDeltaTime >= UpdateFrequency || bForce)
+	{
+		AccumulatedDeltaTime = 0.f;
+
+		UE::Zen::FZenWorkspaces RecentWorkspaces;
+		if (TSharedPtr<UE::Zen::FZenServiceInstance> ServiceInstance = ZenServiceInstance.Get())
+		{
+			ServiceInstance->GetWorkspaces(RecentWorkspaces);
+		}
+
+		if (RecentWorkspaces == Workspaces)
+		{
+			return;
+		}
+
+		Workspaces = RecentWorkspaces;
+		if (GridSlot)
+		{
+			(*GridSlot)
+				[
+					GetWorkspaceList()
+				];
+		}
+
+		SlatePrepass(GetPrepassLayoutScaleMultiplier());
+	}
+}
diff --git a/Engine/Source/Editor/StorageServerWidgets/StorageServerWidgets.Build.cs b/Engine/Source/Editor/StorageServerWidgets/StorageServerWidgets.Build.cs
--- a/Engine/Source/Editor/StorageServerWidgets/StorageServerWidgets.Build.cs
+++ b/Engine/Source/Editor/StorageServerWidgets/StorageServerWidgets.Build.cs
@@ -18,6 +18,7 @@ public class StorageServerWidgets : ModuleRules
 				"Slate",
 				"SlateCore",
 				"Zen",
+				"DesktopPlatform"
 			});
 
 		CppCompileWarningSettings.UnsafeTypeCastWarningLevel = WarningLevel.Error;
diff --git a/Engine/Source/Programs/ZenDashboard/Private/ZenDashboard.cpp b/Engine/Source/Programs/ZenDashboard/Private/ZenDashboard.cpp
--- a/Engine/Source/Programs/ZenDashboard/Private/ZenDashboard.cpp
+++ b/Engine/Source/Programs/ZenDashboard/Private/ZenDashboard.cpp
@@ -22,6 +22,7 @@
 #include "SZenCidStoreStatistics.h"
 #include "SZenProjectStatistics.h"
 #include "SZenServiceStatus.h"
+#include "SZenWorkspaces.h"
 #include "Tasks/Task.h"
 #include "Templates/SharedPointer.h"
 #include "Widgets/Input/SButton.h"
@@ -89,6 +90,7 @@
 	TSharedPtr<SNotificationItem> CompileNotification;
 	TSharedPtr<UE::Zen::FServiceInstanceManager> ServiceInstanceManager;
 	TArray<FSimpleDelegate> MainThreadTasks;
+	TSharedPtr<SZenWorkspaces> ZenWorkspacesWidget;
 
 	std::atomic<bool> bLatentExclusiveOperationActive = false;
 
@@ -418,6 +420,17 @@
 							NAME_None,
 							EUserInterfaceActionType::Button
 						);
+						SubMenuBuilder.AddMenuEntry(
+							LOCTEXT("ResetStats_ZenServer", "Reset stat counters"),
+							LOCTEXT("ResetStats_ZenServer_ToolTip", "Reset the cache stats"),
+							FSlateIcon(),
+							FUIAction(
+								FExecuteAction::CreateRaw(this, &FZenDashboardApp::DeleteDataAndRestartZenServer),
+								FCanExecuteAction::CreateRaw(this, &FZenDashboardApp::CanExecuteExclusiveAction)
+							),
+							NAME_None,
+							EUserInterfaceActionType::Button
+						);
 
 						return SubMenuBuilder.MakeWidget();
 					})
@@ -454,6 +467,22 @@
 		}
 		MenuBuilder.EndSection();
 
+		MenuBuilder.BeginSection(NAME_None, LOCTEXT("Workspace", "Workspace"));
+		{
+			MenuBuilder.AddMenuEntry(
+				LOCTEXT("LaunchUtilityCommandPrompt", "Manage workspaces"),
+				LOCTEXT("LaunchUtilityCommandPrompt_ToolTip", "Manage Zen workspaces: list, create, delete"),
+				FSlateIcon(),
+				FUIAction(
+					FExecuteAction::CreateRaw(this, &FZenDashboardApp::LaunchUtilityCommandPrompt),
+					FCanExecuteAction()
+				),
+				NAME_None,
+				EUserInterfaceActionType::Button
+			);
+		}
+		MenuBuilder.EndSection();
+
 		MenuBuilder.BeginSection(NAME_None, LOCTEXT("Disk", "Disk"));
 		{
 			MenuBuilder.AddMenuEntry(
@@ -525,67 +554,10 @@
 			.FocusWhenFirstShown(false)
 			.SupportsMaximize(false)
 			.SupportsMinimize(true)
-			.HasCloseButton(true)
-			[
-				SNew(SBorder)
-				.BorderImage(FCoreStyle::Get().GetBrush("ToolPanel.GroupBorder"))
-				[
-					SNew(SVerticalBox)
+			.HasCloseButton(true);
 
-					// Menu
-					+ SVerticalBox::Slot()
-					.AutoHeight()
-					[
-						MakeMainMenu()
-					]
+		TSharedPtr<SWidget> WindowContent = CreateWindowContent();
+		Window->SetContent(WindowContent.ToSharedRef());
 
-					// Status panel
-					+ SVerticalBox::Slot()
-					.AutoHeight()
-					.Padding(0.0f, 10.0f, 0.0f, 0.0f)
-					.VAlign(VAlign_Top)
-					[
-						SNew(SZenServiceStatus)
-						.ZenServiceInstance(ServiceInstanceManager.ToSharedRef(), &UE::Zen::FServiceInstanceManager::GetZenServiceInstance)
-					]
-
-					// Stats panel
-					+ SVerticalBox::Slot()
-					.AutoHeight()
-					.Padding(0.0f, 10.0f, 0.0f, 0.0f)
-					.VAlign(VAlign_Top)
-					[
-						SNew(SHorizontalBox)
-
-						+SHorizontalBox::Slot()
-						.AutoWidth()
-						.HAlign(HAlign_Left)
-						.VAlign(VAlign_Top)
-						[
-							SNew(SZenCacheStatistics)
-							.ZenServiceInstance(ServiceInstanceManager.ToSharedRef(), &UE::Zen::FServiceInstanceManager::GetZenServiceInstance)
-						]
-
-						+SHorizontalBox::Slot()
-						.AutoWidth()
-						.HAlign(HAlign_Left)
-						.VAlign(VAlign_Top)
-						[
-							SNew(SZenProjectStatistics)
-							.ZenServiceInstance(ServiceInstanceManager.ToSharedRef(), &UE::Zen::FServiceInstanceManager::GetZenServiceInstance)
-						]
-
-						+SHorizontalBox::Slot()
-						.AutoWidth()
-						.HAlign(HAlign_Left)
-						.VAlign(VAlign_Top)
-						[
-							SNew(SZenCidStoreStatistics)
-							.ZenServiceInstance(ServiceInstanceManager.ToSharedRef(), &UE::Zen::FServiceInstanceManager::GetZenServiceInstance)
-						]
-					]
-				]
-			];
-
 		if (bSystemTrayMode)
 		{
@@ -676,6 +648,100 @@
 	{
 		FCoreDelegates::ModalMessageDialog.Unbind();
 	}
+
+	TSharedPtr<SWidget> CreateWindowContent()
+	{
+		TSharedPtr<SWidget> WindowContent = 
+			SNew(SBorder)
+				.BorderImage(FCoreStyle::Get().GetBrush("ToolPanel.GroupBorder"))
+				[
+					SNew(SVerticalBox)
+
+					// Menu
+					+ SVerticalBox::Slot()
+					.AutoHeight()
+					[
+						MakeMainMenu()
+					]
+
+					// Status panel
+					+ SVerticalBox::Slot()
+					.AutoHeight()
+					.Padding(0.0f, 10.0f, 0.0f, 0.0f)
+					.VAlign(VAlign_Top)
+					[
+						SNew(SZenServiceStatus)
+						.ZenServiceInstance(ServiceInstanceManager.ToSharedRef(), &UE::Zen::FServiceInstanceManager::GetZenServiceInstance)
+					]
+
+					// Stats panel
+					+ SVerticalBox::Slot()
+					.AutoHeight()
+					.Padding(0.0f, 10.0f, 0.0f, 0.0f)
+					.VAlign(VAlign_Top)
+					[
+						SNew(SHorizontalBox)
+
+						+SHorizontalBox::Slot()
+						.AutoWidth()
+						.HAlign(HAlign_Left)
+						.VAlign(VAlign_Top)
+						[
+							SNew(SZenCacheStatistics)
+							.ZenServiceInstance(ServiceInstanceManager.ToSharedRef(), &UE::Zen::FServiceInstanceManager::GetZenServiceInstance)
+						]
+
+						+SHorizontalBox::Slot()
+						.AutoWidth()
+						.HAlign(HAlign_Left)
+						.VAlign(VAlign_Top)
+						[
+							SNew(SZenProjectStatistics)
+							.ZenServiceInstance(ServiceInstanceManager.ToSharedRef(), &UE::Zen::FServiceInstanceManager::GetZenServiceInstance)
+						]
+
+						+SHorizontalBox::Slot()
+						.AutoWidth()
+						.HAlign(HAlign_Left)
+						.VAlign(VAlign_Top)
+						[
+							SNew(SZenCidStoreStatistics)
+							.ZenServiceInstance(ServiceInstanceManager.ToSharedRef(), &UE::Zen::FServiceInstanceManager::GetZenServiceInstance)
+						]
+					]
+
+					// Workspaces panel
+					+ SVerticalBox::Slot()
+					.MaxHeight(300)
+					.Padding(4.f, 10.f, 4.f, 0.f)
+					.FillHeight(1.0f)
+					.VAlign(VAlign_Top)
+					[
+						SAssignNew(ZenWorkspacesWidget, SZenWorkspaces)
+						.ZenServiceInstance(ServiceInstanceManager.ToSharedRef(), &UE::Zen::FServiceInstanceManager::GetZenServiceInstance)
+						.UpdateFrequency(5.f)
+						.Visibility_Lambda([this]() -> EVisibility
+						{
+							static double LastTime = FPlatformTime::Seconds();
+							const double CurrentTime = FPlatformTime::Seconds();
+							const double DeltaTime = CurrentTime - LastTime;
+							LastTime = CurrentTime;
+
+							bool bShow = false;
+							if (ServiceInstanceManager->GetZenServiceInstance().IsValid()
+								&& ServiceInstanceManager->GetZenServiceInstance()->IsServiceRunning())
+							{
+								ZenWorkspacesWidget->UpdateWorkspaces(DeltaTime);
+								bShow = ZenWorkspacesWidget->GetWorkspaceCount() > 0;
+							}
+							
+							return bShow ? EVisibility::Visible : EVisibility::Collapsed;
+						})
+					]
+				];
+
+		return WindowContent;
+	}
 };
 
 int ZenDashboardMain(const TCHAR* CmdLine)
diff --git a/Engine/Source/Programs/ZenDashboard/Private/ZenDashboardStyle.cpp b/Engine/Source/Programs/ZenDashboard/Private/ZenDashboardStyle.cpp
--- a/Engine/Source/Programs/ZenDashboard/Private/ZenDashboardStyle.cpp
+++ b/Engine/Source/Programs/ZenDashboard/Private/ZenDashboardStyle.cpp
@@ -13,9 +13,11 @@
 
 #define EDITOR_IMAGE_BRUSH(RelativePath, ...) IMAGE_BRUSH("../Editor/Slate/" RelativePath, __VA_ARGS__)
 #define EDITOR_IMAGE_BRUSH_SVG(RelativePath, ...) IMAGE_BRUSH_SVG("../Editor/Slate/" RelativePath, __VA_ARGS__)
+#define ENGINE_IMAGE_BRUSH_SVG(RelativePath, ...) IMAGE_BRUSH_SVG("../Slate/" RelativePath, __VA_ARGS__)
 #define EDITOR_BOX_BRUSH(RelativePath, ...) BOX_BRUSH("../Editor/Slate/" RelativePath, __VA_ARGS__)
 #define EDITOR_BORDER_BRUSH(RelativePath, ...) BORDER_BRUSH("../Editor/Slate/" RelativePath, __VA_ARGS__)
 #define RootToContentDir ContentFromEngine
+#define RootToCoreContentDir ContentFromEngine
 
 TSharedPtr< FSlateStyleSet > FZenDashboardStyle::StyleSet = nullptr;
 
@@ -54,6 +56,7 @@ TSharedRef< FSlateStyleSet > FZenDashboardStyle::Create()
 	const ISlateStyle* ParentStyle = FSlateStyleRegistry::FindSlateStyle("CoreStyle");
 	const FTextBlockStyle NormalText = ParentStyle->GetWidgetStyle<FTextBlockStyle>("NormalText");
 
+	const FVector2D Icon10x10(10.f, 10.f);
 	const FVector2D Icon12x12(12.0f, 12.0f);
 	const FVector2D Icon16x16(16.0f, 16.0f);
 
@@ -63,6 +66,19 @@ TSharedRef< FSlateStyleSet > FZenDashboardStyle::Create()
 	Style.Set("Zen.Icons.FolderExplore", new EDITOR_IMAGE_BRUSH_SVG("Starship/Common/ContentBrowser", Icon12x12));
 	Style.Set("Zen.Icons.Clipboard", new EDITOR_IMAGE_BRUSH_SVG("Starship/Common/Clipboard", Icon16x16));
 
+	Style.Set("Icons.ChevronRight", new CORE_IMAGE_BRUSH_SVG("Starship/Common/chevron-right", Icon16x16));
+	Style.Set("Icons.ChevronDown", new CORE_IMAGE_BRUSH_SVG("Starship/Common/chevron-down", Icon16x16));
+
+	Style.Set("GenericCommands.Delete", new CORE_IMAGE_BRUSH_SVG("Starship/Common/Delete", Icon12x12));
+	Style.Set("Icons.Delete-small", new CORE_IMAGE_BRUSH_SVG("Starship/Common/Delete", Icon10x10));
+	Style.Set("Icons.Cross", new CORE_IMAGE_BRUSH_SVG("Starship/Common/close-small", Icon16x16));
+	Style.Set("Icons.XCircle", new CORE_IMAGE_BRUSH_SVG("Starship/Common/close-circle", Icon12x12));
+	Style.Set("Icons.ArrowLeft", new CORE_IMAGE_BRUSH_SVG("Starship/Common/arrow-left", Icon16x16));
+	Style.Set("Icons.FolderOpen", new CORE_IMAGE_BRUSH_SVG("Starship/Common/folder-open", Icon16x16));
+	Style.Set("Icons.Check", new CORE_IMAGE_BRUSH_SVG("Starship/Common/check", Icon16x16));
+
+	Style.Set("Icons.Add", new CORE_IMAGE_BRUSH_SVG("Starship/Common/plus", Icon16x16));
+
 	const FTextBlockStyle DefaultText = FTextBlockStyle()
 		.SetFont(DEFAULT_FONT("Bold", 10))
 		.SetColorAndOpacity(FSlateColor::UseForeground())
