Fix missing #undef LOCTEXT_NAMESPACE in ZenDashboard workspace widget

#rnx

diff --git a/Engine/Source/Editor/StorageServerWidgets/Private/SZenWorkspaces.cpp b/Engine/Source/Editor/StorageServerWidgets/Private/SZenWorkspaces.cpp
--- a/Engine/Source/Editor/StorageServerWidgets/Private/SZenWorkspaces.cpp
+++ b/Engine/Source/Editor/StorageServerWidgets/Private/SZenWorkspaces.cpp
@@ -20,8 +20,6 @@
 #include "DesktopPlatformModule.h"
 #include "Misc/MonitoredProcess.h"
 
-//#include "ZenDashboardStyle.h"
-
 #define LOCTEXT_NAMESPACE "ZenDashboard"
 
 DEFINE_LOG_CATEGORY_STATIC(LogZenDashboardWorkpaces, Log, All);
@@ -724,3 +722,5 @@ void SZenWorkspaces::UpdateWorkspaces(float InDeltaTime, bool bForce)
 		SlatePrepass(GetPrepassLayoutScaleMultiplier());
 	}
 }
+
+#undef LOCTEXT_NAMESPACE
