Fix ZenWorkspace LOCTEXT

#rnx

diff --git a/Engine/Source/Editor/StorageServerWidgets/Private/SZenWorkspaces.cpp b/Engine/Source/Editor/StorageServerWidgets/Private/SZenWorkspaces.cpp
--- a/Engine/Source/Editor/StorageServerWidgets/Private/SZenWorkspaces.cpp
+++ b/Engine/Source/Editor/StorageServerWidgets/Private/SZenWorkspaces.cpp
@@ -294,7 +294,7 @@ class SZenWorkspaceShareRow : public SCompoundWidget
 						.Margin(FMargin(RowMargin, ColumnMargin))
 						.Font(TitleFont)
 						.ColorAndOpacity(FStyleColors::AccentWhite)
-						.Text(LOCTEXT("SharePath", "Alias"))
+						.Text(LOCTEXT("ShareAlias", "Alias"))
 					]
 
 					+ SGridPanel::Slot(1, 2)
