Interim submit of BuildStorageTool code.

diff --git a/Engine/Source/Developer/Zen/Private/BuildServerInterface.cpp b/Engine/Source/Developer/Zen/Private/BuildServerInterface.cpp
new file mode 100644
--- /dev/null
+++ b/Engine/Source/Developer/Zen/Private/BuildServerInterface.cpp
@@ -0,0 +1,560 @@
+// Copyright Epic Games, Inc. All Rights Reserved.
+
+#include "Experimental/BuildServerInterface.h"
+
+#include "Async/Mutex.h"
+#include "Containers/StringView.h"
+#include "DesktopPlatformModule.h"
+#include "Experimental/ZenServerInterface.h"
+#include "HAL/CriticalSection.h"
+#include "HAL/FileManager.h"
+#include "Http/HttpClient.h"
+#include "Http/HttpHostBuilder.h"
+#include "Misc/App.h"
+#include "Misc/ConfigCacheIni.h"
+#include "Misc/MonitoredProcess.h"
+#include "Misc/Paths.h"
+#include "Misc/ScopeExit.h"
+#include "Misc/ScopeLock.h"
+#include "Misc/ScopeRWLock.h"
+#include "Misc/StringBuilder.h"
+#include "Serialization/Archive.h"
+#include "Serialization/CompactBinary.h"
+#include "Serialization/CompactBinarySerialization.h"
+#include "Serialization/CompactBinaryValidation.h"
+#include "Serialization/CompactBinaryWriter.h"
+#include "String/LexFromString.h"
+#include "Tasks/Task.h"
+
+#if WITH_SSL
+#include "Ssl.h"
+#endif
+
+#define UE_BUILDSERVERINSTANCE_MAX_FAILED_LOGIN_ATTEMPTS 16
+
+namespace UE::Zen::Build
+{
+
+DEFINE_LOG_CATEGORY_STATIC(LogBuildServiceInstance, Log, All);
+
+namespace Private
+{
+class FAccessToken
+{
+public:
+	void SetToken(FStringView Scheme, FStringView Token);
+	inline uint32 GetSerial() const { return Serial.load(std::memory_order_relaxed); }
+	friend FAnsiStringBuilderBase& operator<<(FAnsiStringBuilderBase& Builder, const FAccessToken& Token);
+
+private:
+	mutable FRWLock Lock;
+	TArray<ANSICHAR> Header;
+	std::atomic<uint32> Serial;
+};
+
+FAnsiStringBuilderBase& operator<<(FAnsiStringBuilderBase& Builder, const Private::FAccessToken& Token)
+{
+	FReadScopeLock ReadLock(Token.Lock);
+	return Builder.Append(Token.Header);
+}
+
+void FAccessToken::SetToken(const FStringView Scheme, const FStringView Token)
+{
+	FWriteScopeLock WriteLock(Lock);
+	const int32 TokenLen = FPlatformString::ConvertedLength<ANSICHAR>(Token.GetData(), Token.Len());
+
+	Header.Empty(TokenLen);
+
+	const int32 TokenIndex = Header.AddUninitialized(TokenLen);
+	FPlatformString::Convert(Header.GetData() + TokenIndex, TokenLen, Token.GetData(), Token.Len());
+	Serial.fetch_add(1, std::memory_order_relaxed);
+}
+} // namespace Private
+
+bool LoadFromCompactBinary(FCbFieldView Field, FBuildServiceInstance::FBuildRecord& OutBuildRecord)
+{
+	if (!Field.IsObject())
+	{
+		return false;
+	}
+	bool bOk = true;
+	FString BuildIdString;
+	FCbObjectId::ByteArray BuildIdBytes;
+	bOk &= LoadFromCompactBinary(Field["buildId"], BuildIdString);
+	bOk &= UE::String::HexToBytes(BuildIdString, BuildIdBytes) == sizeof(BuildIdBytes);
+	OutBuildRecord.BuildId = FCbObjectId(BuildIdBytes);
+
+	FCbFieldView MetadataField = Field["metadata"];
+	if (!MetadataField.IsObject())
+	{
+		return false;
+	}
+
+	OutBuildRecord.Metadata = FCbObject::Clone(MetadataField.AsObjectView());
+	return bOk;
+}
+
+bool
+FServiceSettings::ReadFromConfig()
+{
+	check(GConfig && GConfig->IsReadyForUse());
+	FString Config;
+	if (!GConfig->GetString(TEXT("StorageServers"), TEXT("Cloud"), Config, GEngineIni))
+	{
+		return false;
+	}
+	bool bRetVal = false;
+
+	bRetVal |= FParse::Value(*Config, TEXT("Host="), Host);
+	bRetVal |= FParse::Value(*Config, TEXT("OAuthProviderIdentifier="), OAuthProviderIdentifier);
+
+	FParse::Value(*Config, TEXT("AuthScheme="), AuthScheme);
+	if (AuthScheme.IsEmpty())
+	{
+		AuthScheme = "Bearer";
+	}
+
+	return bRetVal;
+}
+
+bool
+FServiceSettings::ReadFromURL(FStringView InstanceURL)
+{
+	Host = InstanceURL;
+	return true;
+}
+
+FBuildServiceInstance::FBuildServiceInstance()
+{
+	Settings.ReadFromConfig();
+	Initialize();
+}
+
+FBuildServiceInstance::FBuildServiceInstance(FStringView InstanceURL)
+{
+	Settings.ReadFromURL(InstanceURL);
+	Initialize();
+}
+
+FBuildServiceInstance::~FBuildServiceInstance()
+{
+}
+
+void
+FBuildServiceInstance::Connect(FOnConnectionComplete&& OnConnectionComplete)
+{
+	EConnectionState ExpectedConnectionState = EConnectionState::NotStarted;
+	if (!ConnectionState.compare_exchange_weak(ExpectedConnectionState, EConnectionState::ConnectionInProgress))
+	{
+		if (OnConnectionComplete)
+		{
+			OnConnectionComplete(ExpectedConnectionState);
+		}
+		return;
+	}
+
+	UE::Tasks::Launch(UE_SOURCE_LOCATION, [this, OnConnectionComplete = MoveTemp(OnConnectionComplete)]()
+	{
+		TAnsiStringBuilder<256> ResolvedHost;
+		double ResolvedLatency;
+		FHttpHostBuilder HostBuilder;
+		HostBuilder.AddFromString(Settings.GetHost());
+		if (HostBuilder.ResolveHost(/* Warning timeout */ 1.0, 4.0 /* Max duration timeout*/, ResolvedHost, ResolvedLatency))
+		{
+			EffectiveDomain.Reset();
+			EffectiveDomain.Append(ResolvedHost);
+
+			if (AcquireAccessToken())
+			{
+				ConnectionState.store(EConnectionState::ConnectionSucceeded, std::memory_order_relaxed);
+				if (OnConnectionComplete)
+				{
+					OnConnectionComplete(EConnectionState::ConnectionSucceeded);
+				}
+			}
+			else
+			{
+				UE_LOG(LogBuildServiceInstance, Warning, TEXT("Unable to acquire access token."));
+				ConnectionState.store(EConnectionState::ConnectionFailed, std::memory_order_relaxed);
+				if (OnConnectionComplete)
+				{
+					OnConnectionComplete(EConnectionState::ConnectionFailed);
+				}
+			}
+		}
+		else
+		{
+			// even if we fail to resolve a host to use the returned host will at least contain the first of the possible hosts which we can attempt to use
+			EffectiveDomain.Reset();
+			EffectiveDomain.Append(ResolvedHost);
+
+			FString HostCandidates = HostBuilder.GetHostCandidatesString();
+			UE_LOG(LogBuildServiceInstance, Warning, TEXT("Unable to resolve best host candidate to use, most likely none of the suggested hosts was reachable. Attempted hosts were: '%s' ."), *HostCandidates);
+			ConnectionState.store(EConnectionState::ConnectionFailed, std::memory_order_relaxed);
+			if (OnConnectionComplete)
+			{
+				OnConnectionComplete(EConnectionState::ConnectionFailed);
+			}
+		}
+
+	});
+}
+
+void
+FBuildServiceInstance::RefreshNamespacesAndBuckets(FOnRefreshNamespacesAndBucketsComplete&& InOnRefreshNamespacesAndBucketsComplete)
+{
+	{
+		FWriteScopeLock _(NamespacesAndBucketsLock);
+		NamespacesAndBuckets.Empty();
+	}
+	if (ConnectionState.load(std::memory_order_relaxed) != EConnectionState::ConnectionSucceeded)
+	{
+		CallOnRefreshNamespacesAndBucketsComplete(MoveTemp(InOnRefreshNamespacesAndBucketsComplete));
+		return;
+	}
+
+	FString OutputFilePath = FPaths::CreateTempFilename(FPlatformProcess::UserTempDir(), TEXT("BuildContainers"), TEXT(".cbo"));
+	FPaths::MakePlatformFilename(OutputFilePath);
+
+	FString CommandLineArgs = FString::Printf(TEXT("builds list-namespaces --recursive \"%s\""),
+		*OutputFilePath);
+
+	InvokeZenUtility(CommandLineArgs, nullptr,
+	[this, OutputFilePath = MoveTemp(OutputFilePath), InOnRefreshNamespacesAndBucketsComplete = MoveTemp(InOnRefreshNamespacesAndBucketsComplete)](bool bSuccessful) mutable
+	{
+		if (bSuccessful)
+		{
+			TUniquePtr<FArchive> Ar(IFileManager::Get().CreateFileReader(*OutputFilePath));
+			if (!Ar)
+			{
+				UE_LOG(LogBuildServiceInstance, Warning, TEXT("Missing output file from zen utility when gathering build data: '%s'"), *OutputFilePath);
+				CallOnRefreshNamespacesAndBucketsComplete(MoveTemp(InOnRefreshNamespacesAndBucketsComplete));
+				return;
+			}
+
+			FCbObject OutputObject = LoadCompactBinary(*Ar).AsObject();
+			{
+				FWriteScopeLock _(NamespacesAndBucketsLock);
+				for (FCbFieldView ResultField : OutputObject["results"])
+				{
+					FString Namespace = *WriteToString<64>(ResultField["name"].AsString());
+					for (FCbFieldView Item : ResultField["items"])
+					{
+						NamespacesAndBuckets.Add(Namespace, *WriteToString<64>(Item.AsString()));
+					}
+				}
+			}
+		}
+		CallOnRefreshNamespacesAndBucketsComplete(MoveTemp(InOnRefreshNamespacesAndBucketsComplete));
+	});
+}
+
+void
+FBuildServiceInstance::ListBuilds(FStringView Namespace, FStringView Bucket, FOnListBuildsComplete&& InOnListBuildsComplete)
+{
+	if (ConnectionState.load(std::memory_order_relaxed) != EConnectionState::ConnectionSucceeded)
+	{
+		if (InOnListBuildsComplete)
+		{
+			InOnListBuildsComplete({});
+		}
+		return;
+	}
+
+	FString QueryFilePath = FPaths::CreateTempFilename(FPlatformProcess::UserTempDir(), TEXT("BuildListQuery"), TEXT(".cbo"));
+	FPaths::MakePlatformFilename(QueryFilePath);
+
+	FString OutputFilePath = FPaths::CreateTempFilename(FPlatformProcess::UserTempDir(), TEXT("BuildList"), TEXT(".cbo"));
+	FPaths::MakePlatformFilename(OutputFilePath);
+
+	FString CommandLineArgs = FString::Printf(TEXT("builds list --namespace \"%s\" --bucket \"%s\" \"%s\" \"%s\""),
+		*WriteToString<64>(Namespace), *WriteToString<64>(Bucket), *QueryFilePath, *OutputFilePath);
+
+	auto PreInvoke = [QueryFilePath](FString& CommandLineArgs)
+	{
+		TUniquePtr<FArchive> QueryFileArchive(IFileManager::Get().CreateFileWriter(*QueryFilePath, FILEWRITE_NoFail));
+		FCbWriter Writer;
+		Writer.BeginObject("query");
+		Writer.EndObject();
+		Writer.Save(*QueryFileArchive);
+	};
+
+	InvokeZenUtility(CommandLineArgs, MoveTemp(PreInvoke),
+	[OutputFilePath = MoveTemp(OutputFilePath), InOnListBuildsComplete = MoveTemp(InOnListBuildsComplete)](bool bSuccessful) mutable
+	{
+		if (!bSuccessful)
+		{
+			if (InOnListBuildsComplete)
+			{
+				InOnListBuildsComplete({});
+			}
+			return;
+		}
+
+		TUniquePtr<FArchive> Ar(IFileManager::Get().CreateFileReader(*OutputFilePath));
+		if (!Ar)
+		{
+			UE_LOG(LogBuildServiceInstance, Warning, TEXT("Missing output file from zen utility when gathering build data: '%s'"), *OutputFilePath);
+			if (InOnListBuildsComplete)
+			{
+				InOnListBuildsComplete({});
+			}
+			return;
+		}
+
+		if (InOnListBuildsComplete)
+		{
+			FCbObject OutputObject = LoadCompactBinary(*Ar).AsObject();
+			TArray<FBuildRecord> BuildRecords;
+			for (FCbFieldView ResultField : OutputObject["results"].AsArrayView())
+			{
+				FBuildRecord BuildRecord;
+				if (LoadFromCompactBinary(ResultField, BuildRecord))
+				{
+					BuildRecords.Emplace(MoveTemp(BuildRecord));
+				}
+			}
+			InOnListBuildsComplete(MoveTemp(BuildRecords));
+		}
+
+	});
+}
+
+FBuildServiceInstance::EConnectionState
+FBuildServiceInstance::GetConnectionState() const
+{
+	return ConnectionState.load(std::memory_order_relaxed);
+}
+
+FAnsiStringView
+FBuildServiceInstance::GetEffectiveDomain() const
+{
+	return EffectiveDomain;
+}
+
+void
+FBuildServiceInstance::Initialize()
+{
+#if WITH_SSL
+	// Load SSL module during HTTP module's StatupModule() to make sure module manager figures out the dependencies correctly
+	// and doesn't unload SSL before unloading HTTP module at exit
+	FSslModule::Get();
+#endif
+	FDesktopPlatformModule::TryGet();
+}
+
+bool
+FBuildServiceInstance::AcquireAccessToken()
+{
+	if (GetEffectiveDomain().StartsWith("http://localhost"))
+	{
+		UE_LOG(LogBuildServiceInstance, Log, TEXT("Skipping authorization for connection to localhost."));
+		return true;
+	}
+
+	LoginAttempts++;
+
+	// Avoid spamming this if the service is down.
+	if (FailedLoginAttempts > UE_BUILDSERVERINSTANCE_MAX_FAILED_LOGIN_ATTEMPTS)
+	{
+		return false;
+	}
+
+	TRACE_CPUPROFILER_EVENT_SCOPE(BuildServiceInstance_AcquireAccessToken);
+
+	// In case many requests wants to update the token at the same time
+	// get the current serial while we wait to take the CS.
+	const uint32 WantsToUpdateTokenSerial = Access ? Access->GetSerial() : 0;
+
+	FScopeLock Lock(&AccessCs);
+
+	// If the token was updated while we waited to take the lock, then it should now be valid.
+	if (Access && Access->GetSerial() > WantsToUpdateTokenSerial)
+	{
+		return true;
+	}
+
+	FString AccessTokenString;
+	FDateTime TokenExpiresAt;
+	bool bWasInteractiveLogin = false;
+
+	IDesktopPlatform* DesktopPlatform = FDesktopPlatformModule::TryGet();
+	if (DesktopPlatform && DesktopPlatform->GetOidcAccessToken(FPaths::RootDir(), {}, Settings.GetOAuthProviderIdentifier(), FApp::IsUnattended(), GWarn, AccessTokenString, TokenExpiresAt, bWasInteractiveLogin))
+	{
+		if (bWasInteractiveLogin)
+		{
+			InteractiveLoginAttempts++;
+		}
+
+		const double ExpiryTimeSeconds = (TokenExpiresAt - FDateTime::UtcNow()).GetTotalSeconds();
+		UE_LOG(LogBuildServiceInstance, Display,
+			TEXT("OidcToken: Logged in to HTTP DDC services. Expires at %s which is in %.0f seconds."),
+			*TokenExpiresAt.ToString(), ExpiryTimeSeconds);
+		SetAccessTokenAndUnlock(Lock, AccessTokenString, ExpiryTimeSeconds);
+		return true;
+	}
+	else if (DesktopPlatform)
+	{
+		UE_LOG(LogBuildServiceInstance, Warning, TEXT("OidcToken: Failed to log in to HTTP services."));
+		FailedLoginAttempts++;
+		return false;
+	}
+	else
+	{
+		UE_LOG(LogBuildServiceInstance, Warning, TEXT("OidcToken: Use of OAuthProviderIdentifier requires that the target depend on DesktopPlatform."));
+		FailedLoginAttempts++;
+		return false;
+	}
+	return false;
+}
+
+void
+FBuildServiceInstance::SetAccessTokenAndUnlock(FScopeLock& Lock, FStringView Token, double RefreshDelay)
+{
+	// Cache the expired refresh handle.
+	FTSTicker::FDelegateHandle ExpiredRefreshAccessTokenHandle = MoveTemp(RefreshAccessTokenHandle);
+	RefreshAccessTokenHandle.Reset();
+
+	if (!Access)
+	{
+		Access = MakeUnique<Private::FAccessToken>();
+	}
+	Access->SetToken(Settings.GetAuthScheme(), Token);
+
+	constexpr double RefreshGracePeriod = 20.0f;
+	if (RefreshDelay > RefreshGracePeriod)
+	{
+		// Schedule a refresh of the token ahead of expiry time (this will not work in commandlets)
+		if (!IsRunningCommandlet())
+		{
+			RefreshAccessTokenHandle = FTSTicker::GetCoreTicker().AddTicker(FTickerDelegate::CreateLambda(
+				[this](float DeltaTime)
+				{
+					AcquireAccessToken();
+					return false;
+				}
+			), float(FMath::Min(RefreshDelay - RefreshGracePeriod, MAX_flt)));
+		}
+
+		// Schedule a forced refresh of the token when the scheduled refresh is starved or unavailable.
+		RefreshAccessTokenTime = FPlatformTime::Seconds() + RefreshDelay - RefreshGracePeriod * 0.5f;
+	}
+	else
+	{
+		RefreshAccessTokenTime = 0.0;
+	}
+
+	// Reset failed login attempts, the service is indeed alive.
+	FailedLoginAttempts = 0;
+
+	// Unlock the critical section before attempting to remove the expired refresh handle.
+	// The associated ticker delegate could already be executing, which could cause a
+	// hang in RemoveTicker when the critical section is locked.
+	Lock.Unlock();
+	if (ExpiredRefreshAccessTokenHandle.IsValid())
+	{
+		FTSTicker::GetCoreTicker().RemoveTicker(MoveTemp(ExpiredRefreshAccessTokenHandle));
+	}
+}
+
+FString
+FBuildServiceInstance::GetAccessToken() const
+{
+	TAnsiStringBuilder<128> AccessTokenBuilder;
+	if (Access.IsValid())
+	{
+		AccessTokenBuilder << *Access;
+	}
+	return FString(AccessTokenBuilder);
+}
+
+void
+FBuildServiceInstance::InvokeZenUtility(FStringView InCommandLineArgs, FPreZenUtilityInvocation&& PreInvocation, FOnZenUtilityInvocationComplete&& OnComplete)
+{
+	UE::Tasks::Launch(UE_SOURCE_LOCATION,
+	[this, InCommandLineArgs = FString(InCommandLineArgs), PreInvocation = MoveTemp(PreInvocation), OnComplete = MoveTemp(OnComplete)]() mutable
+	{
+ 		FString ZenUtilityPath = UE::Zen::GetLocalInstallUtilityPath();
+
+		FString CommandLineArgs = FString::Printf(TEXT("%s --host %s --access-token \"%s\""),
+			*InCommandLineArgs, *WriteToString<64>(EffectiveDomain), *GetAccessToken());
+
+		if (PreInvocation)
+		{
+			PreInvocation(CommandLineArgs);
+		}
+
+		FString AbsoluteUtilityPath = FPaths::ConvertRelativePathToFull(ZenUtilityPath);
+		FMonitoredProcess MonitoredUtilityProcess(AbsoluteUtilityPath, *CommandLineArgs, FPaths::GetPath(AbsoluteUtilityPath), true);
+		if (!MonitoredUtilityProcess.Launch())
+		{
+			UE_LOG(LogBuildServiceInstance, Warning, TEXT("Failed to launch zen utility to gather build data: '%s'."), *AbsoluteUtilityPath);
+			if (OnComplete)
+			{
+				OnComplete(false);
+			}
+			return;
+		}
+
+		const uint64 StartTime = FPlatformTime::Cycles64();
+		while (MonitoredUtilityProcess.Update())
+		{
+			double Duration = FPlatformTime::ToSeconds64(FPlatformTime::Cycles64() - StartTime);
+			if (Duration > 120.0)
+			{
+				MonitoredUtilityProcess.Cancel(true);
+				UE_LOG(LogBuildServiceInstance, Warning, TEXT("Cancelled launch of zen utility for gathering build data: '%s' due to timeout."), *AbsoluteUtilityPath);
+
+				// Wait for execution to be terminated
+				while (MonitoredUtilityProcess.Update())
+				{
+					Duration = FPlatformTime::ToSeconds64(FPlatformTime::Cycles64() - StartTime);
+					if (Duration > 15.0)
+					{
+						UE_LOG(LogBuildServiceInstance, Warning, TEXT("Cancelled launch of zen utility for gathering build data: '%s'. Failed waiting for termination."), *AbsoluteUtilityPath);
+						break;
+					}
+					FPlatformProcess::Sleep(0.2f);
+				}
+
+				FString OutputString = MonitoredUtilityProcess.GetFullOutputWithoutDelegate();
+				UE_LOG(LogBuildServiceInstance, Warning, TEXT("Launch of zen utility for gathering build data: '%s' failed. Output: '%s'"), *AbsoluteUtilityPath, *OutputString);
+				if (OnComplete)
+				{
+					OnComplete(false);
+				}
+				return;
+			}
+			FPlatformProcess::Sleep(0.1f);
+		}
+
+		FString OutputString = MonitoredUtilityProcess.GetFullOutputWithoutDelegate();
+		if (MonitoredUtilityProcess.GetReturnCode() != 0)
+		{
+			UE_LOG(LogBuildServiceInstance, Warning, TEXT("Unexpected return code after launch of zen utility for gathering build data: '%s' (%d). Output: '%s'"), *AbsoluteUtilityPath, MonitoredUtilityProcess.GetReturnCode(), *OutputString);
+			if (OnComplete)
+			{
+				OnComplete(false);
+			}
+			return;
+		}
+
+		if (OnComplete)
+		{
+			OnComplete(true);
+		}
+	});
+}
+
+void
+FBuildServiceInstance::CallOnRefreshNamespacesAndBucketsComplete(FOnRefreshNamespacesAndBucketsComplete&& InOnRefreshNamespacesAndBucketsComplete)
+{
+	if (InOnRefreshNamespacesAndBucketsComplete)
+	{
+		InOnRefreshNamespacesAndBucketsComplete();
+	}
+
+	RefreshNamespacesAndBucketsComplete.Broadcast();
+}
+
+} // namespace UE::StorageService::Build
diff --git a/Engine/Source/Developer/Zen/Public/Experimental/BuildServerInterface.h b/Engine/Source/Developer/Zen/Public/Experimental/BuildServerInterface.h
new file mode 100644
--- /dev/null
+++ b/Engine/Source/Developer/Zen/Public/Experimental/BuildServerInterface.h
@@ -0,0 +1,108 @@
+// Copyright Epic Games, Inc. All Rights Reserved.
+
+#pragma once
+
+#include "Containers/Map.h"
+#include "Containers/StringFwd.h"
+#include "Containers/StringView.h"
+#include "Containers/Ticker.h"
+#include "Containers/UnrealString.h"
+#include "HAL/CriticalSection.h"
+#include "Serialization/CompactBinary.h"
+#include "Templates/Function.h"
+#include <atomic>
+
+#define UE_API ZEN_API
+
+namespace UE { class IHttpClient; }
+namespace UE::Zen::Build::Private { class FAccessToken; }
+
+namespace UE::Zen::Build
+{
+
+struct FServiceSettings
+{
+public:
+	UE_API bool ReadFromConfig();
+	UE_API bool ReadFromURL(FStringView InstanceURL);
+
+	inline const FString& GetHost() const { return Host; }
+	inline const FString& GetAuthScheme() const { return AuthScheme; }
+	inline const FString& GetOAuthProviderIdentifier() const { return OAuthProviderIdentifier; }
+private:
+	FString Host;
+	FString AuthScheme;
+	FString OAuthProviderIdentifier;
+};
+
+class FBuildServiceInstance
+{
+public:
+	enum class EConnectionState
+	{
+		NotStarted,
+		ConnectionInProgress,
+		ConnectionSucceeded,
+		ConnectionFailed
+	};
+
+	struct FBuildRecord
+	{
+		FCbObjectId BuildId;
+		FCbObject Metadata;
+	};
+
+	typedef TUniqueFunction<void(EConnectionState)> FOnConnectionComplete;
+	typedef TUniqueFunction<void()> FOnRefreshNamespacesAndBucketsComplete;
+	typedef TUniqueFunction<void(TArray<FBuildRecord>&& ListBuildRecords)> FOnListBuildsComplete;
+
+
+	UE_API FBuildServiceInstance();
+	UE_API FBuildServiceInstance(FStringView InstanceURL);
+
+	UE_API virtual ~FBuildServiceInstance();
+
+	UE_API void Connect(FOnConnectionComplete&& OnConnectionComplete = {});
+	UE_API void RefreshNamespacesAndBuckets(FOnRefreshNamespacesAndBucketsComplete&& InOnRefreshNamespacesAndBucketsComplete = {});
+	[[nodiscard]] UE_API EConnectionState GetConnectionState() const;
+	[[nodiscard]] UE_API FAnsiStringView GetEffectiveDomain() const;
+
+	[[nodiscard]] UE_API TMulticastDelegateRegistration<void()>& OnRefreshNamespacesAndBucketsComplete() { return RefreshNamespacesAndBucketsComplete; }
+
+	inline const TMultiMap<FString, FString>& GetNamespacesAndBuckets() const { return NamespacesAndBuckets; }
+
+	UE_API void ListBuilds(FStringView Namespace, FStringView Bucket, FOnListBuildsComplete&& InOnListBuildsComplete);
+private:
+	typedef TUniqueFunction<void(FString& CommandLineArgs)> FPreZenUtilityInvocation;
+	typedef TUniqueFunction<void(bool)> FOnZenUtilityInvocationComplete;
+
+	void Initialize();
+	bool AcquireAccessToken();
+	void SetAccessTokenAndUnlock(FScopeLock &Lock, FStringView Token, double RefreshDelay = 0.0);
+	[[nodiscard]] FString GetAccessToken() const;
+	void InvokeZenUtility(FStringView CommandlineArgs, FPreZenUtilityInvocation&& PreInvocation, FOnZenUtilityInvocationComplete&& OnComplete);
+	void CallOnRefreshNamespacesAndBucketsComplete(FOnRefreshNamespacesAndBucketsComplete&& InOnRefreshNamespacesAndBucketsComplete);
+
+	TMulticastDelegate<void()> RefreshNamespacesAndBucketsComplete;
+	std::atomic<EConnectionState> ConnectionState = EConnectionState::NotStarted;
+
+	FCriticalSection AccessCs;
+	TUniquePtr<Private::FAccessToken> Access;
+	FTSTicker::FDelegateHandle RefreshAccessTokenHandle;
+	double RefreshAccessTokenTime = 0.0;
+	uint32 LoginAttempts = 0;
+	uint32 FailedLoginAttempts = 0;
+	uint32 InteractiveLoginAttempts = 0;
+
+	FRWLock NamespacesAndBucketsLock;
+	TMultiMap<FString, FString> NamespacesAndBuckets;
+
+	FServiceSettings Settings;
+	FAnsiStringBuilderBase EffectiveDomain;
+};
+
+bool LoadFromCompactBinary(FCbFieldView Field, FBuildServiceInstance::FBuildRecord& OutBuildRecord);
+
+} // namespace UE::StorageService::Build
+
+#undef UE_API
diff --git a/Engine/Source/Editor/StorageServerWidgets/Internal/BuildServiceInstanceManager.h b/Engine/Source/Editor/StorageServerWidgets/Internal/BuildServiceInstanceManager.h
new file mode 100644
--- /dev/null
+++ b/Engine/Source/Editor/StorageServerWidgets/Internal/BuildServiceInstanceManager.h
@@ -0,0 +1,21 @@
+// Copyright Epic Games, Inc. All Rights Reserved.
+
+#pragma once
+
+#include "Templates/SharedPointer.h"
+
+namespace UE::Zen::Build { class FBuildServiceInstance; }
+
+namespace UE::Zen::Build
+{
+
+class STORAGESERVERWIDGETS_API FServiceInstanceManager
+{
+public:
+	TSharedPtr<FBuildServiceInstance> GetBuildServiceInstance() const;
+private:
+	mutable TSharedPtr<FBuildServiceInstance> CurrentInstance;
+};
+
+} // namespace UE::StorageService::Build
+
diff --git a/Engine/Source/Editor/StorageServerWidgets/Internal/SBuildActivity.h b/Engine/Source/Editor/StorageServerWidgets/Internal/SBuildActivity.h
new file mode 100644
--- /dev/null
+++ b/Engine/Source/Editor/StorageServerWidgets/Internal/SBuildActivity.h
@@ -0,0 +1,39 @@
+// Copyright Epic Games, Inc. All Rights Reserved.
+
+#pragma once
+
+#include "Experimental/BuildServerInterface.h"
+#include "Experimental/ZenServerInterface.h"
+#include "Templates/SharedPointer.h"
+#include "Widgets/DeclarativeSyntaxSupport.h"
+#include "Widgets/SBoxPanel.h"
+#include "Widgets/SCompoundWidget.h"
+#include "Widgets/SWidget.h"
+
+class STORAGESERVERWIDGETS_API SBuildActivity : public SCompoundWidget
+{
+	SLATE_BEGIN_ARGS(SBuildActivity)
+		: _ZenServiceInstance(nullptr)
+		, _BuildServiceInstance(nullptr)
+	{ }
+
+	SLATE_ATTRIBUTE(TSharedPtr<UE::Zen::FZenServiceInstance>, ZenServiceInstance);
+	SLATE_ATTRIBUTE(TSharedPtr<UE::Zen::Build::FBuildServiceInstance>, BuildServiceInstance);
+
+	SLATE_END_ARGS()
+
+	struct FBuildActivityItem : public TSharedFromThis<FBuildActivityItem, ESPMode::ThreadSafe>
+	{
+	};
+
+	void Construct(const FArguments& InArgs);
+
+private:
+	TSharedRef<SWidget> GetGridPanel();
+
+	SVerticalBox::FSlot* GridSlot = nullptr;
+	TAttribute<TSharedPtr<UE::Zen::FZenServiceInstance>> ZenServiceInstance;
+	TAttribute<TSharedPtr<UE::Zen::Build::FBuildServiceInstance>> BuildServiceInstance;
+};
+
+typedef TSharedPtr<SBuildActivity::FBuildActivityItem, ESPMode::ThreadSafe> FBuildActivityItemPtr;
\ No newline at end of file
diff --git a/Engine/Source/Editor/StorageServerWidgets/Internal/SBuildDestination.h b/Engine/Source/Editor/StorageServerWidgets/Internal/SBuildDestination.h
new file mode 100644
--- /dev/null
+++ b/Engine/Source/Editor/StorageServerWidgets/Internal/SBuildDestination.h
@@ -0,0 +1,33 @@
+// Copyright Epic Games, Inc. All Rights Reserved.
+
+#pragma once
+
+#include "Experimental/BuildServerInterface.h"
+#include "Experimental/ZenServerInterface.h"
+#include "Templates/SharedPointer.h"
+#include "Widgets/DeclarativeSyntaxSupport.h"
+#include "Widgets/SBoxPanel.h"
+#include "Widgets/SCompoundWidget.h"
+#include "Widgets/SWidget.h"
+
+class STORAGESERVERWIDGETS_API SBuildDestination : public SCompoundWidget
+{
+	SLATE_BEGIN_ARGS(SBuildDestination)
+		: _ZenServiceInstance(nullptr)
+		, _BuildServiceInstance(nullptr)
+	{ }
+
+	SLATE_ATTRIBUTE(TSharedPtr<UE::Zen::FZenServiceInstance>, ZenServiceInstance);
+	SLATE_ATTRIBUTE(TSharedPtr<UE::Zen::Build::FBuildServiceInstance>, BuildServiceInstance);
+
+	SLATE_END_ARGS()
+
+	void Construct(const FArguments& InArgs);
+
+private:
+	TSharedRef<SWidget> GetGridPanel();
+
+	SVerticalBox::FSlot* GridSlot = nullptr;
+	TAttribute<TSharedPtr<UE::Zen::FZenServiceInstance>> ZenServiceInstance;
+	TAttribute<TSharedPtr<UE::Zen::Build::FBuildServiceInstance>> BuildServiceInstance;
+};
\ No newline at end of file
diff --git a/Engine/Source/Editor/StorageServerWidgets/Internal/SBuildLogin.h b/Engine/Source/Editor/StorageServerWidgets/Internal/SBuildLogin.h
new file mode 100644
--- /dev/null
+++ b/Engine/Source/Editor/StorageServerWidgets/Internal/SBuildLogin.h
@@ -0,0 +1,33 @@
+// Copyright Epic Games, Inc. All Rights Reserved.
+
+#pragma once
+
+#include "Experimental/BuildServerInterface.h"
+#include "Experimental/ZenServerInterface.h"
+#include "Templates/SharedPointer.h"
+#include "Widgets/DeclarativeSyntaxSupport.h"
+#include "Widgets/SBoxPanel.h"
+#include "Widgets/SCompoundWidget.h"
+#include "Widgets/SWidget.h"
+
+class STORAGESERVERWIDGETS_API SBuildLogin : public SCompoundWidget
+{
+	SLATE_BEGIN_ARGS(SBuildLogin)
+		: _ZenServiceInstance(nullptr)
+		, _BuildServiceInstance(nullptr)
+	{ }
+
+	SLATE_ATTRIBUTE(TSharedPtr<UE::Zen::FZenServiceInstance>, ZenServiceInstance);
+	SLATE_ATTRIBUTE(TSharedPtr<UE::Zen::Build::FBuildServiceInstance>, BuildServiceInstance);
+
+	SLATE_END_ARGS()
+
+	void Construct(const FArguments& InArgs);
+
+private:
+	TSharedRef<SWidget> GetGridPanel();
+
+	SVerticalBox::FSlot* GridSlot = nullptr;
+	TAttribute<TSharedPtr<UE::Zen::FZenServiceInstance>> ZenServiceInstance;
+	TAttribute<TSharedPtr<UE::Zen::Build::FBuildServiceInstance>> BuildServiceInstance;
+};
diff --git a/Engine/Source/Editor/StorageServerWidgets/Internal/SBuildSelection.h b/Engine/Source/Editor/StorageServerWidgets/Internal/SBuildSelection.h
new file mode 100644
--- /dev/null
+++ b/Engine/Source/Editor/StorageServerWidgets/Internal/SBuildSelection.h
@@ -0,0 +1,77 @@
+// Copyright Epic Games, Inc. All Rights Reserved.
+
+#pragma once
+
+#include "Experimental/BuildServerInterface.h"
+#include "Experimental/ZenServerInterface.h"
+#include "Templates/SharedPointer.h"
+#include "Widgets/DeclarativeSyntaxSupport.h"
+#include "Widgets/Input/SComboBox.h"
+#include "Widgets/SBoxPanel.h"
+#include "Widgets/SCompoundWidget.h"
+#include "Widgets/SWidget.h"
+#include "Widgets/Views/ITableRow.h"
+#include "Widgets/Views/STableViewBase.h"
+
+class SMultiSelectComboBox;
+
+namespace UE::BuildSelection::Internal{ struct FListBuildsState; }
+
+class STORAGESERVERWIDGETS_API SBuildSelection : public SCompoundWidget
+{
+	SLATE_BEGIN_ARGS(SBuildSelection)
+		: _ZenServiceInstance(nullptr)
+		, _BuildServiceInstance(nullptr)
+	{ }
+
+	SLATE_ATTRIBUTE(TSharedPtr<UE::Zen::FZenServiceInstance>, ZenServiceInstance);
+	SLATE_ATTRIBUTE(TSharedPtr<UE::Zen::Build::FBuildServiceInstance>, BuildServiceInstance);
+
+	SLATE_END_ARGS()
+
+	struct FBuildGroup : public TSharedFromThis<FBuildGroup, ESPMode::ThreadSafe>
+	{
+		FString Namespace;
+		FString DisplayName;
+		FString CommitIdentifier;
+		FString Category;
+		TMap<FString, UE::Zen::Build::FBuildServiceInstance::FBuildRecord> PerPlatformBuilds;
+	};
+
+	void Construct(const FArguments& InArgs);
+
+private:
+	void RebuildLists();
+	void RegenerateBuildGroups(UE::BuildSelection::Internal::FListBuildsState& ListBuildsState);
+	TSharedRef<SWidget> OnGenerateTextBlockFromString(TSharedPtr<FString> Item);
+	TSharedRef<ITableRow> GenerateBuildGroupRow(TSharedPtr<FBuildGroup> InItem, const TSharedRef<STableViewBase>& InOwningTable);
+	TSharedRef<SWidget> GetGridPanel();
+
+	TSharedPtr<SComboBox<TSharedPtr<FString>>> ProjectWidget;
+	TSharedPtr<SComboBox<TSharedPtr<FString>>> StreamWidget;
+	TSharedPtr<SComboBox<TSharedPtr<FString>>> BuildTypeWidget;
+	TSharedPtr<SMultiSelectComboBox> RequiredPlatformsWidget;
+
+	TMultiMap<FString, FString> BucketsToNamespaces;
+
+	TArray<TSharedPtr<FString>> ProjectList;
+	TArray<TSharedPtr<FString>> StreamList;
+	TArray<TSharedPtr<FString>> BuildTypeList;
+	TArray<TSharedPtr<FString>> PlatformList;
+
+	TSharedPtr<FString> SelectedStream;
+	TSharedPtr<FString> SelectedProject;
+	TSharedPtr<FString> SelectedBuildType;
+
+	TSharedPtr<SListView<TSharedPtr<SBuildSelection::FBuildGroup, ESPMode::ThreadSafe>>> BuildListView;
+
+	std::atomic<uint32> BuildListRefreshesInProgress = 0;
+	std::atomic<uint32> BuildRefreshGeneration = 0;
+	TArray<TSharedPtr<SBuildSelection::FBuildGroup, ESPMode::ThreadSafe>> BuildGroups;
+
+	SVerticalBox::FSlot* GridSlot = nullptr;
+	TAttribute<TSharedPtr<UE::Zen::FZenServiceInstance>> ZenServiceInstance;
+	TAttribute<TSharedPtr<UE::Zen::Build::FBuildServiceInstance>> BuildServiceInstance;
+};
+
+typedef TSharedPtr<SBuildSelection::FBuildGroup, ESPMode::ThreadSafe> FBuildSelectionBuildGroupPtr;
\ No newline at end of file
diff --git a/Engine/Source/Editor/StorageServerWidgets/Private/BuildServiceInstanceManager.cpp b/Engine/Source/Editor/StorageServerWidgets/Private/BuildServiceInstanceManager.cpp
new file mode 100644
--- /dev/null
+++ b/Engine/Source/Editor/StorageServerWidgets/Private/BuildServiceInstanceManager.cpp
@@ -0,0 +1,19 @@
+// Copyright Epic Games, Inc. All Rights Reserved.
+
+#include "BuildServiceInstanceManager.h"
+#include "Experimental/BuildServerInterface.h"
+
+
+namespace UE::Zen::Build
+{
+
+TSharedPtr<FBuildServiceInstance> FServiceInstanceManager::GetBuildServiceInstance() const
+{
+	if (!CurrentInstance)
+	{
+		CurrentInstance = MakeShared<FBuildServiceInstance>();
+	}
+	return CurrentInstance;
+}
+
+} // namespace UE::StorageService::Build
diff --git a/Engine/Source/Editor/StorageServerWidgets/Private/SBuildActivity.cpp b/Engine/Source/Editor/StorageServerWidgets/Private/SBuildActivity.cpp
new file mode 100644
--- /dev/null
+++ b/Engine/Source/Editor/StorageServerWidgets/Private/SBuildActivity.cpp
@@ -0,0 +1,67 @@
+// Copyright Epic Games, Inc. All Rights Reserved.
+
+#include "SBuildActivity.h"
+#include "Internationalization/FastDecimalFormat.h"
+#include "Math/BasicMathExpressionEvaluator.h"
+#include "Math/UnitConversion.h"
+#include "Misc/ExpressionParser.h"
+#include "Styling/StyleColors.h"
+#include "Widgets/Images/SImage.h"
+#include "Widgets/Input/SButton.h"
+#include "Widgets/Input/SHyperlink.h"
+#include "Widgets/Layout/SGridPanel.h"
+#include "Widgets/SBoxPanel.h"
+#include "Widgets/Text/STextBlock.h"
+#include "Widgets/Views/SListView.h"
+#include "ZenServiceInstanceManager.h"
+
+#define LOCTEXT_NAMESPACE "StorageServerBuild"
+
+void SBuildActivity::Construct(const FArguments& InArgs)
+{
+	ZenServiceInstance = InArgs._ZenServiceInstance;
+	BuildServiceInstance = InArgs._BuildServiceInstance;
+
+	this->ChildSlot
+	[
+		SNew(SVerticalBox)
+
+		+ SVerticalBox::Slot()
+		.AutoHeight()
+		.Padding(0, 0, 0, 0)
+		.Expose(GridSlot)
+		[
+			GetGridPanel()
+		]
+	];
+}
+
+TSharedRef<SWidget> SBuildActivity::GetGridPanel()
+{
+	TSharedRef<SGridPanel> Panel = SNew(SGridPanel);
+
+	const float RowMargin = 0.0f;
+	const float ColumnMargin = 10.0f;
+	const FSlateColor TitleColor = FStyleColors::AccentWhite;
+	const FSlateFontInfo TitleFont = FCoreStyle::GetDefaultFontStyle("Bold", 10);
+
+	Panel->AddSlot(0, 0)
+	[
+		SNew(STextBlock)
+		.Margin(FMargin(ColumnMargin, RowMargin))
+		.ColorAndOpacity(TitleColor)
+		.Font(TitleFont)
+		.Text(LOCTEXT("BuildActivity_Activity", "Activity"))
+	];
+
+	Panel->AddSlot(0, 1)
+	.Padding(FMargin(ColumnMargin, RowMargin))
+	[
+		SNew(SListView<FBuildActivityItemPtr>)
+	];
+
+
+	return Panel;
+}
+
+#undef LOCTEXT_NAMESPACE
diff --git a/Engine/Source/Editor/StorageServerWidgets/Private/SBuildDestination.cpp b/Engine/Source/Editor/StorageServerWidgets/Private/SBuildDestination.cpp
new file mode 100644
--- /dev/null
+++ b/Engine/Source/Editor/StorageServerWidgets/Private/SBuildDestination.cpp
@@ -0,0 +1,64 @@
+// Copyright Epic Games, Inc. All Rights Reserved.
+
+#include "SBuildDestination.h"
+#include "Internationalization/FastDecimalFormat.h"
+#include "Math/BasicMathExpressionEvaluator.h"
+#include "Math/UnitConversion.h"
+#include "Misc/ExpressionParser.h"
+#include "Styling/StyleColors.h"
+#include "Widgets/Images/SImage.h"
+#include "Widgets/Input/SButton.h"
+#include "Widgets/Input/SHyperlink.h"
+#include "Widgets/Layout/SGridPanel.h"
+#include "Widgets/SBoxPanel.h"
+#include "Widgets/Text/STextBlock.h"
+#include "ZenServiceInstanceManager.h"
+
+#define LOCTEXT_NAMESPACE "StorageServerBuild"
+
+void SBuildDestination::Construct(const FArguments& InArgs)
+{
+	ZenServiceInstance = InArgs._ZenServiceInstance;
+	BuildServiceInstance = InArgs._BuildServiceInstance;
+
+	this->ChildSlot
+	[
+		SNew(SVerticalBox)
+
+		+ SVerticalBox::Slot()
+		.AutoHeight()
+		.Padding(0, 0, 0, 0)
+		.Expose(GridSlot)
+		[
+			GetGridPanel()
+		]
+	];
+}
+
+TSharedRef<SWidget> SBuildDestination::GetGridPanel()
+{
+	TSharedRef<SGridPanel> Panel = SNew(SGridPanel);
+
+	int32 Row = 0;
+
+	const float RowMargin = 0.0f;
+	const float ColumnMargin = 10.0f;
+	const FSlateColor TitleColor = FStyleColors::AccentWhite;
+	const FSlateFontInfo TitleFont = FCoreStyle::GetDefaultFontStyle("Bold", 10);
+
+	Panel->AddSlot(0, Row)
+	[
+		SNew(STextBlock)
+		.Margin(FMargin(ColumnMargin, RowMargin))
+		.ColorAndOpacity(TitleColor)
+		.Font(TitleFont)
+		.Text(LOCTEXT("ServiceStatus_Status", "TODO: Destination"))
+	];
+
+	++Row;
+
+
+	return Panel;
+}
+
+#undef LOCTEXT_NAMESPACE
diff --git a/Engine/Source/Editor/StorageServerWidgets/Private/SBuildLogin.cpp b/Engine/Source/Editor/StorageServerWidgets/Private/SBuildLogin.cpp
new file mode 100644
--- /dev/null
+++ b/Engine/Source/Editor/StorageServerWidgets/Private/SBuildLogin.cpp
@@ -0,0 +1,134 @@
+// Copyright Epic Games, Inc. All Rights Reserved.
+
+#include "SBuildLogin.h"
+#include "Internationalization/FastDecimalFormat.h"
+#include "Math/BasicMathExpressionEvaluator.h"
+#include "Math/UnitConversion.h"
+#include "Misc/ExpressionParser.h"
+#include "Styling/StyleColors.h"
+#include "Widgets/Images/SImage.h"
+#include "Widgets/Input/SButton.h"
+#include "Widgets/Input/SHyperlink.h"
+#include "Widgets/Layout/SGridPanel.h"
+#include "Widgets/SBoxPanel.h"
+#include "Widgets/Text/STextBlock.h"
+#include "Widgets/Views/SListView.h"
+#include "ZenServiceInstanceManager.h"
+
+#define LOCTEXT_NAMESPACE "StorageServerBuild"
+
+void SBuildLogin::Construct(const FArguments& InArgs)
+{
+	ZenServiceInstance = InArgs._ZenServiceInstance;
+	BuildServiceInstance = InArgs._BuildServiceInstance;
+
+	this->ChildSlot
+	[
+		SNew(SVerticalBox)
+
+		+ SVerticalBox::Slot()
+		.AutoHeight()
+		.Padding(0, 0, 0, 0)
+		.Expose(GridSlot)
+		[
+			GetGridPanel()
+		]
+	];
+}
+
+TSharedRef<SWidget> SBuildLogin::GetGridPanel()
+{
+	TSharedRef<SGridPanel> Panel = SNew(SGridPanel);
+
+	const float RowMargin = 0.0f;
+	const float ColumnMargin = 10.0f;
+	const FSlateColor TitleColor = FStyleColors::AccentWhite;
+	const FSlateFontInfo TitleFont = FCoreStyle::GetDefaultFontStyle("Bold", 10);
+
+	Panel->AddSlot(0, 0)
+	[
+		SNew(STextBlock)
+		.Margin(FMargin(ColumnMargin, RowMargin))
+		.ColorAndOpacity(TitleColor)
+		.Font(TitleFont)
+		.Text(LOCTEXT("BuildConnect_StatusLabel", "Status:"))
+	];
+
+	Panel->AddSlot(1, 0)
+	[
+		SNew(STextBlock)
+		.Text_Lambda([this]
+		{
+			if (TSharedPtr<UE::Zen::Build::FBuildServiceInstance> ServiceInstance = BuildServiceInstance.Get())
+			{
+				switch (ServiceInstance->GetConnectionState())
+				{
+				case UE::Zen::Build::FBuildServiceInstance::EConnectionState::NotStarted:
+					return LOCTEXT("BuildConnect_StatusValueNotConnected", "Not Connected");
+				case UE::Zen::Build::FBuildServiceInstance::EConnectionState::ConnectionInProgress:
+					return LOCTEXT("BuildConnect_StatusValueConnecting", "Connecting...");
+				case UE::Zen::Build::FBuildServiceInstance::EConnectionState::ConnectionSucceeded:
+					{
+						return FText::Format(LOCTEXT("BuildConnect_StatusValueConnected", "Connected to {0}"),
+							FText::FromStringView(ServiceInstance->GetEffectiveDomain()));
+					}
+				case UE::Zen::Build::FBuildServiceInstance::EConnectionState::ConnectionFailed:
+					return LOCTEXT("BuildConnect_StatusValueConnectionFailed", "Connection failed");
+				}
+			}
+			return LOCTEXT("BuildConnect_StatusValueError", "Error");
+		})
+		.Visibility_Lambda([this]
+		{
+			if (TSharedPtr<UE::Zen::Build::FBuildServiceInstance> ServiceInstance = BuildServiceInstance.Get())
+			{
+				if (ServiceInstance->GetConnectionState() == UE::Zen::Build::FBuildServiceInstance::EConnectionState::NotStarted)
+				{
+					return EVisibility::Collapsed;
+				}
+			}
+			return EVisibility::Visible;
+		})
+	];
+
+	Panel->AddSlot(2, 0)
+	[
+		SNew(SHyperlink)
+		.Padding(FMargin(ColumnMargin, RowMargin))
+		.UnderlineStyle(FAppStyle::Get(), "NoBorder")
+		.Text_Lambda([this]
+		{
+			FText BrowseText = LOCTEXT("BuildLogin_ConnectLink", "Click to connect");
+			return BrowseText;
+		})
+		.OnNavigate_Lambda([this]
+		{
+			if (TSharedPtr<UE::Zen::Build::FBuildServiceInstance> ServiceInstance = BuildServiceInstance.Get())
+			{
+				ServiceInstance->Connect([ServiceInstance](UE::Zen::Build::FBuildServiceInstance::EConnectionState ConnectionState)
+					{
+						if (ConnectionState == UE::Zen::Build::FBuildServiceInstance::EConnectionState::ConnectionSucceeded)
+						{
+							ServiceInstance->RefreshNamespacesAndBuckets();
+						}
+					});
+			}
+		})
+		.Visibility_Lambda([this]
+		{
+			if (TSharedPtr<UE::Zen::Build::FBuildServiceInstance> ServiceInstance = BuildServiceInstance.Get())
+			{
+				if (ServiceInstance->GetConnectionState() == UE::Zen::Build::FBuildServiceInstance::EConnectionState::NotStarted)
+				{
+					return EVisibility::Visible;
+				}
+			}
+			return EVisibility::Hidden;
+		})
+	];
+
+
+	return Panel;
+}
+
+#undef LOCTEXT_NAMESPACE
diff --git a/Engine/Source/Editor/StorageServerWidgets/Private/SBuildSelection.cpp b/Engine/Source/Editor/StorageServerWidgets/Private/SBuildSelection.cpp
new file mode 100644
--- /dev/null
+++ b/Engine/Source/Editor/StorageServerWidgets/Private/SBuildSelection.cpp
@@ -0,0 +1,663 @@
+// Copyright Epic Games, Inc. All Rights Reserved.
+
+#include "SBuildSelection.h"
+
+#include "Internationalization/FastDecimalFormat.h"
+#include "Math/BasicMathExpressionEvaluator.h"
+#include "Math/UnitConversion.h"
+#include "Misc/ExpressionParser.h"
+#include "SMultiSelectComboBox.h"
+#include "Styling/StyleColors.h"
+#include "Widgets/Images/SImage.h"
+#include "Widgets/Input/SButton.h"
+#include "Widgets/Input/SHyperlink.h"
+#include "Widgets/Layout/SGridPanel.h"
+#include "Widgets/SBoxPanel.h"
+#include "Widgets/Text/STextBlock.h"
+#include "Widgets/Views/SHeaderRow.h"
+#include "Widgets/Views/SListView.h"
+#include "ZenServiceInstanceManager.h"
+#include <atomic>
+
+#define LOCTEXT_NAMESPACE "StorageServerBuild"
+
+namespace UE::BuildSelection::Internal
+{
+	namespace FBuildGroupIds
+	{
+		const FName ColName = TEXT("Name");
+		const FName ColCommit = TEXT("Commit");
+		const FName ColSuffix = TEXT("Suffix");
+		const FName ColCreated = TEXT("Created");
+	}
+
+	struct FNamespacePlatformBucketTuple
+	{
+		FString Namespace;
+		FString Platform;
+		FString Bucket;
+	};
+
+	struct FBuildState
+	{
+		FString Namespace;
+		FString Platform;
+		TArray<UE::Zen::Build::FBuildServiceInstance::FBuildRecord> Results;
+	};
+
+	struct FListBuildsState
+	{
+		std::atomic<uint32> PendingQueries;
+		TArray<FBuildState> QueryState;
+	};
+}
+
+void SBuildSelection::Construct(const FArguments& InArgs)
+{
+	ZenServiceInstance = InArgs._ZenServiceInstance;
+	BuildServiceInstance = InArgs._BuildServiceInstance;
+
+	if (TSharedPtr<UE::Zen::Build::FBuildServiceInstance> ServiceInstance = BuildServiceInstance.Get())
+	{
+		ServiceInstance->OnRefreshNamespacesAndBucketsComplete().AddSP(this, &SBuildSelection::RebuildLists);
+	}
+
+	this->ChildSlot
+	[
+		SNew(SVerticalBox)
+
+		+ SVerticalBox::Slot()
+		.HAlign(HAlign_Fill)
+		.VAlign(VAlign_Fill)
+		.Padding(0, 0, 0, 0)
+		.Expose(GridSlot)
+		[
+			GetGridPanel()
+		]
+	];
+}
+
+void SBuildSelection::RebuildLists()
+{
+	using namespace UE::Zen::Build;
+	using namespace UE::BuildSelection::Internal;
+
+	BuildListRefreshesInProgress.fetch_add(1);
+	BucketsToNamespaces.Empty();
+	StreamList.Empty();
+	ProjectList.Empty();
+	BuildTypeList.Empty();
+	PlatformList.Empty();
+
+	if (TSharedPtr<FBuildServiceInstance> ServiceInstance = BuildServiceInstance.Get())
+	{
+		TArray<FString> Namespaces;
+		TArray<FString> Projects;
+		TArray<FString> Streams;
+		TArray<FString> BuildTypes;
+		TArray<FString> Platforms;
+		const TMultiMap<FString, FString>& NamespacesAndBuckets = ServiceInstance->GetNamespacesAndBuckets();
+
+		NamespacesAndBuckets.GetKeys(Namespaces);
+
+		for (const TPair<FString, FString>& NamespaceAndBucket : NamespacesAndBuckets)
+		{
+			BucketsToNamespaces.AddUnique(NamespaceAndBucket.Value, NamespaceAndBucket.Key);
+
+			TArray<FString> Segments;
+			NamespaceAndBucket.Value.ParseIntoArray(Segments, TEXT("."), true);
+
+			if (Segments.Num() == 4)
+			{
+				Projects.AddUnique(Segments[0]);
+				if (Segments[0] != (SelectedProject ? *SelectedProject : Projects[0]))
+				{
+					continue;
+				}
+
+				Streams.AddUnique(Segments[2]);
+				if (Segments[2] != (SelectedStream ? *SelectedStream : Streams[0]))
+				{
+					continue;
+				}
+
+				BuildTypes.AddUnique(Segments[1]);
+				if (Segments[1] != (SelectedBuildType ? *SelectedBuildType : BuildTypes[0]))
+				{
+					continue;
+				}
+
+				Platforms.AddUnique(Segments[3]);
+			}
+		}
+
+		auto ConvertToSharedPtrs = [] (TArray<FString>& Strings, TArray<TSharedPtr<FString>>& SharedStrings)
+		{
+			Strings.StableSort();
+			for (FString& String : Strings)
+			{
+				SharedStrings.Add(MakeShared<FString>(MoveTemp(String)));
+			}
+		};
+
+		ConvertToSharedPtrs(Projects, ProjectList);
+		ConvertToSharedPtrs(Streams, StreamList);
+		ConvertToSharedPtrs(BuildTypes, BuildTypeList);
+		ConvertToSharedPtrs(Platforms, PlatformList);
+	}
+
+	auto ConformSelection = [](TSharedPtr<FString>& SelectedItem, const TArray<TSharedPtr<FString>>& SelectionList)
+	{
+		if (!SelectedItem)
+		{
+			SelectedItem = SelectionList.IsEmpty() ? nullptr : SelectionList[0];
+			return;
+		}
+
+		const TSharedPtr<FString>* FoundSelectionListItem = SelectionList.FindByPredicate([&SelectedItem](const TSharedPtr<FString>& Item)
+		{
+			return *Item == *SelectedItem;
+		});
+
+		SelectedItem = FoundSelectionListItem ? *FoundSelectionListItem : SelectionList.IsEmpty() ? nullptr : SelectionList[0];
+	};
+
+	ConformSelection(SelectedProject, ProjectList);
+	ConformSelection(SelectedStream, StreamList);
+	ConformSelection(SelectedBuildType, BuildTypeList);
+
+	ExecuteOnGameThread(UE_SOURCE_LOCATION,
+		[this]
+		{
+			StreamWidget->RefreshOptions();
+			StreamWidget->SetSelectedItem(SelectedStream);
+			ProjectWidget->RefreshOptions();
+			ProjectWidget->SetSelectedItem(SelectedProject);
+			BuildTypeWidget->RefreshOptions();
+			BuildTypeWidget->SetSelectedItem(SelectedBuildType);
+		});
+
+	if (TSharedPtr<FBuildServiceInstance> ServiceInstance = BuildServiceInstance.Get())
+	{
+		TArray<FNamespacePlatformBucketTuple> NamespacePlatformBucketTuples;
+		for (TSharedPtr<FString> Platform : PlatformList)
+		{
+			FString Bucket = FString::Printf(TEXT("%s.%s.%s.%s"), **SelectedProject, **SelectedBuildType, **SelectedStream, **Platform);
+			TArray<FString> NamespacesForBucket;
+			BucketsToNamespaces.MultiFind(Bucket, NamespacesForBucket);
+			for (FString& Namespace : NamespacesForBucket)
+			{
+				NamespacePlatformBucketTuples.Emplace(MoveTemp(Namespace), *Platform, MoveTemp(Bucket));
+			}
+		}
+
+		TSharedPtr<FListBuildsState> PendingQueryState = MakeShared<FListBuildsState>();
+		PendingQueryState->PendingQueries = NamespacePlatformBucketTuples.Num();
+		PendingQueryState->QueryState.SetNum(NamespacePlatformBucketTuples.Num());
+		uint32 QueryIndex = 0;
+
+		++BuildRefreshGeneration;
+		BuildGroups.Empty();
+
+		for (const FNamespacePlatformBucketTuple& NamespacePlatformBucket : NamespacePlatformBucketTuples)
+		{
+			ServiceInstance->ListBuilds(NamespacePlatformBucket.Namespace, NamespacePlatformBucket.Bucket,
+			[this, QueryIndex, Namespace = NamespacePlatformBucket.Namespace, Platform = NamespacePlatformBucket.Platform,
+				ExpectedBuildRefreshGeneration = BuildRefreshGeneration.load(), PendingQueryState]
+			(TArray<FBuildServiceInstance::FBuildRecord>&& Results) mutable
+				{
+					FBuildState& NewBuildState = PendingQueryState->QueryState[QueryIndex];
+					NewBuildState.Namespace = MoveTemp(Namespace);
+					NewBuildState.Platform = MoveTemp(Platform);
+					NewBuildState.Results = MoveTemp(Results);
+
+					if ((ExpectedBuildRefreshGeneration == BuildRefreshGeneration) && (--PendingQueryState->PendingQueries == 0))
+					{
+						// All queries complete
+						RegenerateBuildGroups(*PendingQueryState);
+
+						ExecuteOnGameThread(UE_SOURCE_LOCATION,
+							[this]
+							{
+								BuildListView->RequestListRefresh();
+								BuildListRefreshesInProgress.fetch_sub(1);
+							});
+					}
+				});
+			++QueryIndex;
+		}
+	}
+	else
+	{
+		BuildListRefreshesInProgress.fetch_sub(1);
+	}
+}
+
+void SBuildSelection::RegenerateBuildGroups(UE::BuildSelection::Internal::FListBuildsState& ListBuildsState)
+{
+	using namespace UE::Zen::Build;
+	using namespace UE::BuildSelection::Internal;
+
+	BuildGroups.Empty();
+
+	auto MakeGroupKey = [](const FString& Namespace, const FString& CommitIdentifier, const FBuildServiceInstance::FBuildRecord& BuildRecord)
+	{
+		if (!CommitIdentifier.IsEmpty())
+		{
+			return FString(WriteToString<64>(Namespace, ".", CommitIdentifier));
+		}
+		if (FCbFieldView NameField = BuildRecord.Metadata["name"]; !NameField.HasError())
+		{
+			return FString(WriteToString<64>(Namespace, ".", NameField.AsString()));
+		}
+		return FString(WriteToString<64>(BuildRecord.BuildId));
+	};
+
+	TMap<FString, TSharedPtr<FBuildGroup>> KeyedGroups;
+	for (FBuildState& BuildState : ListBuildsState.QueryState)
+	{
+		for (FBuildServiceInstance::FBuildRecord& BuildRecord : BuildState.Results)
+		{
+			FString CommitIdentifier;
+			if (FCbFieldView ChangelistField = BuildRecord.Metadata["changelist"]; !ChangelistField.HasError())
+			{
+				if (ChangelistField.IsString())
+				{
+					CommitIdentifier = FUTF8ToTCHAR(ChangelistField.AsString());
+				}
+				else if (ChangelistField.IsInteger())
+				{
+					CommitIdentifier = *WriteToString<64>(ChangelistField.AsUInt64());
+				}
+			}
+			else if (FCbFieldView CommitField = BuildRecord.Metadata["commit"]; !CommitField.HasError())
+			{
+				if (CommitField.IsString())
+				{
+					CommitIdentifier = FUTF8ToTCHAR(CommitField.AsString());
+				}
+				else if (CommitField.IsInteger())
+				{
+					CommitIdentifier = *WriteToString<64>(CommitField.AsUInt64());
+				}
+			}
+			FString GroupKey = MakeGroupKey(BuildState.Namespace, CommitIdentifier, BuildRecord);
+			TSharedPtr<FBuildGroup>& BuildGroup = KeyedGroups.FindOrAdd(GroupKey);
+			if (!BuildGroup)
+			{
+				BuildGroup = MakeShared<FBuildGroup>();
+			}
+
+			if (BuildGroup->DisplayName.IsEmpty())
+			{
+				FString Category;
+				if (FCbFieldView TemplateIdField = BuildRecord.Metadata["hordeTemplateId"]; !TemplateIdField.HasError())
+				{
+					Category = FUTF8ToTCHAR(TemplateIdField.AsString());
+				}
+
+				FString ItemName;
+				if (FCbFieldView NameView = BuildRecord.Metadata["name"]; !NameView.HasError())
+				{
+					// TODO: This name manipulation needs to be removed when the metadata is more consistent.
+					ItemName = FUTF8ToTCHAR(NameView.AsString());
+					int32 CLStartIndex = ItemName.Find(TEXT("-CL"));
+					if (CLStartIndex != INDEX_NONE)
+					{
+						int32 TruncationIndex = ItemName.Find(TEXT("-"), ESearchCase::IgnoreCase, ESearchDir::FromStart, CLStartIndex + 4);
+						if (TruncationIndex != INDEX_NONE)
+						{
+							ItemName.LeftInline(TruncationIndex);
+						}
+						TruncationIndex = ItemName.Find(TEXT("."), ESearchCase::IgnoreCase, ESearchDir::FromStart, CLStartIndex + 4);
+						if (TruncationIndex != INDEX_NONE)
+						{
+							ItemName.LeftInline(TruncationIndex);
+						}
+					}
+
+					if (!Category.IsEmpty() && ItemName.StartsWith(Category))
+					{
+						ItemName.RightChopInline(Category.Len());
+					}
+					bool bCharRemoved = false;
+					do
+					{
+						ItemName.TrimCharInline(TCHAR('.'), &bCharRemoved);
+					}
+					while(bCharRemoved);
+					do
+					{
+						ItemName.TrimCharInline(TCHAR('+'), &bCharRemoved);
+					}
+					while(bCharRemoved);
+
+					ItemName.ReplaceCharInline(TCHAR('+'), TCHAR('-'));
+				}
+				else
+				{
+					ItemName = *WriteToString<64>(BuildRecord.BuildId);
+				}
+
+				BuildGroup->DisplayName = ItemName;
+				BuildGroup->CommitIdentifier = CommitIdentifier;
+				BuildGroup->Namespace = BuildState.Namespace;
+			}
+
+			BuildGroup->PerPlatformBuilds.FindOrAdd(BuildState.Platform, MoveTemp(BuildRecord));
+		}
+	}
+
+	KeyedGroups.GenerateValueArray(BuildGroups);
+}
+
+TSharedRef<SWidget> SBuildSelection::OnGenerateTextBlockFromString(TSharedPtr<FString> Item)
+{
+	return SNew(STextBlock)
+		.Text(FText::FromString(*Item));
+}
+
+TSharedRef<ITableRow> SBuildSelection::GenerateBuildGroupRow(FBuildSelectionBuildGroupPtr InItem, const TSharedRef<STableViewBase>& InOwningTable)
+{
+	return SNew(STableRow<FBuildSelectionBuildGroupPtr>, InOwningTable)
+		[
+			SNew(SHorizontalBox)
+				+ SHorizontalBox::Slot()
+				[
+					SNew(STextBlock).Text(FText::FromString(InItem->DisplayName))
+				]
+				+ SHorizontalBox::Slot()
+				[
+					SNew(STextBlock).Text(FText::FromString(InItem->CommitIdentifier))
+				]
+		];
+}
+
+TSharedRef<SWidget> SBuildSelection::GetGridPanel()
+{
+	using namespace UE::BuildSelection::Internal;
+
+	TSharedRef<SVerticalBox> Panel =
+	SNew(SVerticalBox)
+	.IsEnabled_Lambda([this]
+	{
+		if (TSharedPtr<UE::Zen::Build::FBuildServiceInstance> ServiceInstance = BuildServiceInstance.Get())
+		{
+			return ServiceInstance->GetConnectionState() == UE::Zen::Build::FBuildServiceInstance::EConnectionState::ConnectionSucceeded && !ServiceInstance->GetNamespacesAndBuckets().IsEmpty();
+		}
+		return false;
+	});
+
+	const float MinDesiredWidth = 50.0f;
+	int32 Row = 0;
+
+	const float RowMargin = 0.0f;
+	const float ColumnMargin = 10.0f;
+	const FSlateColor TitleColor = FStyleColors::AccentWhite;
+	const FSlateFontInfo TitleFont = FCoreStyle::GetDefaultFontStyle("Bold", 10);
+
+	Panel->AddSlot()
+	.AutoHeight()
+	.HAlign(HAlign_Fill)
+	.VAlign(VAlign_Top)
+	[
+		SNew(SHorizontalBox)
+
+		// Project
+		+SHorizontalBox::Slot()
+		.AutoWidth()
+		.FillWidth(1.0f)
+		.HAlign(HAlign_Fill)
+		.VAlign(VAlign_Top)
+		[
+			SNew(SVerticalBox)
+
+			+ SVerticalBox::Slot()
+			.AutoHeight()
+			.Padding(FMargin(ColumnMargin, RowMargin))
+			.VAlign(VAlign_Top)
+			[
+				SNew(STextBlock)
+				.ColorAndOpacity(TitleColor)
+				.Font(TitleFont)
+				.Text(LOCTEXT("BuildFilter_Project", "Project"))
+			]
+
+			+ SVerticalBox::Slot()
+			.AutoHeight()
+			.Padding(FMargin(ColumnMargin, RowMargin))
+			.HAlign(HAlign_Fill)
+			.VAlign(VAlign_Top)
+			[
+				SAssignNew(ProjectWidget, SComboBox<TSharedPtr<FString>>)
+				.OptionsSource(&ProjectList)
+				.OnSelectionChanged_Lambda([this](TSharedPtr<FString> Item, ESelectInfo::Type SelectInfo)
+				{
+					if (Item.IsValid() && *Item != *SelectedProject)
+					{
+						SelectedProject = Item;
+						SelectedStream = nullptr;
+						SelectedBuildType = nullptr;
+						RebuildLists();
+					}
+				})
+				.OnGenerateWidget(this, &SBuildSelection::OnGenerateTextBlockFromString)
+				[
+					SNew(STextBlock)
+						.MinDesiredWidth(MinDesiredWidth)
+						.Text_Lambda([this]()
+						{
+							return FText::FromString(SelectedProject ? **SelectedProject : TEXT(""));
+						})
+				]
+			]
+		]
+
+		// Stream
+		+SHorizontalBox::Slot()
+		.AutoWidth()
+		.FillWidth(1.0f)
+		.HAlign(HAlign_Fill)
+		.VAlign(VAlign_Top)
+		[
+			SNew(SVerticalBox)
+
+			+ SVerticalBox::Slot()
+			.AutoHeight()
+			.Padding(FMargin(ColumnMargin, RowMargin))
+			.VAlign(VAlign_Top)
+			[
+				SNew(STextBlock)
+				.ColorAndOpacity(TitleColor)
+				.Font(TitleFont)
+				.Text(LOCTEXT("BuildFilter_Stream", "Stream"))
+			]
+
+			+ SVerticalBox::Slot()
+			.AutoHeight()
+			.Padding(FMargin(ColumnMargin, RowMargin))
+			.HAlign(HAlign_Fill)
+			.VAlign(VAlign_Top)
+			[
+				SAssignNew(StreamWidget, SComboBox<TSharedPtr<FString>>)
+				.OptionsSource(&StreamList)
+				.OnSelectionChanged_Lambda([this](TSharedPtr<FString> Item, ESelectInfo::Type SelectInfo)
+				{
+					if (Item.IsValid() && *Item != *SelectedStream)
+					{
+						SelectedStream = Item;
+						SelectedBuildType = nullptr;
+						RebuildLists();
+					}
+				})
+				.OnGenerateWidget(this, &SBuildSelection::OnGenerateTextBlockFromString)
+				[
+					SNew(STextBlock)
+						.MinDesiredWidth(MinDesiredWidth)
+						.Text_Lambda([this]()
+						{
+							return FText::FromString(SelectedStream ? **SelectedStream : TEXT(""));
+						})
+				]
+			]
+		]
+
+		// Build Type
+		+SHorizontalBox::Slot()
+		.AutoWidth()
+		.FillWidth(1.0f)
+		.HAlign(HAlign_Fill)
+		.VAlign(VAlign_Top)
+		[
+			SNew(SVerticalBox)
+
+			+ SVerticalBox::Slot()
+			.AutoHeight()
+			.Padding(FMargin(ColumnMargin, RowMargin))
+			.VAlign(VAlign_Top)
+			[
+				SNew(STextBlock)
+				.ColorAndOpacity(TitleColor)
+				.Font(TitleFont)
+				.Text(LOCTEXT("BuildFilter_BuildType", "Build Type"))
+			]
+
+			+ SVerticalBox::Slot()
+			.AutoHeight()
+			.Padding(FMargin(ColumnMargin, RowMargin))
+			.HAlign(HAlign_Fill)
+			.VAlign(VAlign_Top)
+			[
+				SAssignNew(BuildTypeWidget, SComboBox<TSharedPtr<FString>>)
+				.OptionsSource(&BuildTypeList)
+				.OnSelectionChanged_Lambda([this](TSharedPtr<FString> Item, ESelectInfo::Type SelectInfo)
+				{
+					if (Item.IsValid() && *Item != *SelectedBuildType)
+					{
+						SelectedBuildType = Item;
+						RebuildLists();
+					}
+				})
+				.OnGenerateWidget(this, &SBuildSelection::OnGenerateTextBlockFromString)
+				[
+					SNew(STextBlock)
+						.MinDesiredWidth(MinDesiredWidth)
+						.Text_Lambda([this]()
+						{
+							return FText::FromString(SelectedBuildType ? **SelectedBuildType : TEXT(""));
+						})
+				]
+			]
+		]
+
+		// Required Platforms
+		+SHorizontalBox::Slot()
+		.AutoWidth()
+		.FillWidth(1.0f)
+		.HAlign(HAlign_Fill)
+		.VAlign(VAlign_Top)
+		[
+			SNew(SVerticalBox)
+
+			+ SVerticalBox::Slot()
+			.AutoHeight()
+			.Padding(FMargin(ColumnMargin, RowMargin))
+			.VAlign(VAlign_Top)
+			[
+				SNew(STextBlock)
+				.ColorAndOpacity(TitleColor)
+				.Font(TitleFont)
+				.Text(LOCTEXT("BuildFilter_RequiredPlatforms", "Required Platforms"))
+			]
+
+			+ SVerticalBox::Slot()
+			.AutoHeight()
+			.Padding(FMargin(ColumnMargin, RowMargin))
+			.HAlign(HAlign_Fill)
+			.VAlign(VAlign_Top)
+			[
+				SAssignNew(RequiredPlatformsWidget, SMultiSelectComboBox)
+				.SelectValues(&PlatformList)
+				.OnCheckboxChanged_Lambda([this](ECheckBoxState)
+				{
+				})
+			]
+		]
+	];
+
+	++Row;
+
+	Panel->AddSlot()
+	.Padding(FMargin(ColumnMargin, 10, ColumnMargin, 0))
+	.AutoHeight()
+	.HAlign(HAlign_Fill)
+	.VAlign(VAlign_Top)
+	[
+		SNew(STextBlock)
+			.ColorAndOpacity(TitleColor)
+			.Font(TitleFont)
+			.Text(LOCTEXT("BuildFilter_BuildsLabel", "Builds"))
+	];
+
+	++Row;
+
+	Panel->AddSlot()
+	.Padding(FMargin(ColumnMargin, RowMargin))
+	.HAlign(HAlign_Fill)
+	.VAlign(VAlign_Fill)
+	[
+		SNew(SVerticalBox)
+		+SVerticalBox::Slot()
+		.HAlign(HAlign_Fill)
+		[
+			SAssignNew(BuildListView, SListView<FBuildSelectionBuildGroupPtr>)
+				.ListItemsSource(&BuildGroups)
+				.OnGenerateRow(this, &SBuildSelection::GenerateBuildGroupRow)
+				.SelectionMode(ESelectionMode::Single)
+				.IsEnabled_Lambda([this]
+				{
+					return !BuildListRefreshesInProgress;
+				})
+				.HeaderRow
+				(
+					SNew(SHeaderRow)
+					+ SHeaderRow::Column(FBuildGroupIds::ColName).DefaultLabel(LOCTEXT("BuildFilter_BuildGroupColName", "Name"))
+					+ SHeaderRow::Column(FBuildGroupIds::ColCommit).DefaultLabel(LOCTEXT("BuildFilter_BuildGroupColCommit", "Commit"))
+						.DefaultTooltip(LOCTEXT("BuildFilter_BuildGroupColCommitTooltip", "Commit/Changelist for the build"))
+						.FixedWidth(60).HAlignCell(HAlign_Right).HAlignHeader(HAlign_Center).VAlignCell(VAlign_Center)
+					+ SHeaderRow::Column(FBuildGroupIds::ColSuffix).DefaultLabel(LOCTEXT("BuildFilter_BuildGroupColSuffix", "Suffix"))
+						.DefaultTooltip(LOCTEXT("BuildFilter_BuildGroupColSuffixTooltip", "Modifier on top of the commit/changelist for the build"))
+						.FixedWidth(60).HAlignCell(HAlign_Right).HAlignHeader(HAlign_Center).VAlignCell(VAlign_Center)
+					+ SHeaderRow::Column(FBuildGroupIds::ColCreated).DefaultLabel(LOCTEXT("BuildFilter_BuildGroupColCreated", "Created"))
+						.DefaultTooltip(LOCTEXT("BuildFilter_BuildGroupColCreatedTooltip", "When the build was created"))
+						.FixedWidth(60).HAlignCell(HAlign_Right).HAlignHeader(HAlign_Center).VAlignCell(VAlign_Center)
+				)
+		]
+	];
+
+	++Row;
+
+	Panel->AddSlot()
+	.Padding(FMargin(ColumnMargin, 3, ColumnMargin, 0))
+	.AutoHeight()
+	.HAlign(HAlign_Fill)
+	.VAlign(VAlign_Bottom)
+	[
+		SNew(STextBlock)
+			.Font(FAppStyle::Get().GetFontStyle("SmallFont"))
+			.Text_Lambda([this]()
+			{
+				if (BuildListRefreshesInProgress)
+				{
+					return LOCTEXT("BuildFilter_ResultLoading", "Loading");
+				}
+				return FText::Format(LOCTEXT("BuildFilter_ResultDescription", "{0} items"), FText::AsNumber(BuildGroups.Num()));
+			})
+	];
+
+	++Row;
+
+	return Panel;
+}
+
+#undef LOCTEXT_NAMESPACE
diff --git a/Engine/Source/Editor/StorageServerWidgets/Private/SMultiSelectComboBox.cpp b/Engine/Source/Editor/StorageServerWidgets/Private/SMultiSelectComboBox.cpp
new file mode 100644
--- /dev/null
+++ b/Engine/Source/Editor/StorageServerWidgets/Private/SMultiSelectComboBox.cpp
@@ -0,0 +1,138 @@
+// Copyright Epic Games, Inc. All Rights Reserved.
+
+#include "SMultiSelectComboBox.h"
+
+#include "Widgets/Input/SHyperlink.h"
+#include "Widgets/Input/SButton.h"
+#include "Widgets/Input/SCheckBox.h"
+#include "Widgets/Input/STextComboBox.h"
+#include "Widgets/Input/SEditableTextBox.h"
+#include "Widgets/Views/SListView.h"
+#include "Widgets/Input/SHyperlink.h"
+
+#define LOCTEXT_NAMESPACE "StorageServerBuild"
+
+void SMultiSelectComboBox::Construct(const FArguments& InArgs)
+{
+	using namespace UE::Slate::Containers;
+	SelectValues = InArgs._SelectValues;
+
+	TSharedPtr<SHorizontalBox> HorizontalBox;
+	
+	ChildSlot
+		[
+			SAssignNew(HorizontalBox, SHorizontalBox)
+		];
+
+	HorizontalBox->AddSlot()
+		.HAlign(HAlign_Fill)
+		[
+			SNew(SBox)
+			.MinDesiredWidth(70.0f)
+			[
+				SNew(SComboButton)
+					.IsEnabled_Lambda([this]()
+						{
+							return SelectValues->Num() > 0;
+						})
+				.ButtonContent()
+				[
+					SNew(STextBlock)
+					.Text_Lambda([this]()
+						{
+							bool bFirst = true;
+							TStringBuilder<64> StringBuilder;
+							for (const FString& CheckedValue : CheckedValues)
+							{
+								bool bInSelectValues = SelectValues->ContainsByPredicate([&CheckedValue](const TSharedPtr<FString>& A)
+								{
+									return *A == CheckedValue;
+								});
+								if (bInSelectValues)
+								{
+									if (!bFirst)
+									{
+										StringBuilder << TEXT(", ");
+									}
+									bFirst = false;
+									StringBuilder << CheckedValue;
+								}
+							}
+							return StringBuilder.Len() == 0 ? LOCTEXT("MultiSelectComboBox_AnyValue", "Any") : FText::FromString(*StringBuilder);
+						})
+				]
+				.OnGetMenuContent(FOnGetContent::CreateSP(this, &SMultiSelectComboBox::GenerateMenuItems))
+			]
+		];
+}
+
+SMultiSelectComboBox::~SMultiSelectComboBox()
+{
+}
+
+TSharedRef<class SWidget> SMultiSelectComboBox::GenerateMenuItems()
+{
+	TSharedRef<SVerticalBox> Contents = SNew(SVerticalBox);
+
+	for (int i = 0; i < this->SelectValues->Num(); i++)
+	{
+		Contents->AddSlot()
+			.AutoHeight()
+			.Padding(2)
+			[
+				SNew(SHorizontalBox)
+				+SHorizontalBox::Slot()
+				.AutoWidth()
+				.VAlign(VAlign_Center)
+				[
+					SNew(SCheckBox)
+					.IsChecked_Lambda([this, i]() { return CheckedValues.Contains(*this->SelectValues->operator[](i)) ? ECheckBoxState::Checked : ECheckBoxState::Unchecked; })
+					.OnCheckStateChanged_Lambda([this, i](ECheckBoxState InNewState)
+					{
+						 OnSelectedChangedFromMultiselect(SelectValues->operator[](i));
+					})
+				]
+				+SHorizontalBox::Slot()
+				.AutoWidth()
+				.VAlign(VAlign_Center)
+				[
+					SNew(SButton)
+					.ButtonStyle(FAppStyle::Get(), "InvisibleButton")
+					.IsFocusable(false)			
+
+					.OnClicked_Lambda([this, i]() { OnSelectedChangedFromMultiselect(SelectValues->operator[](i)); return FReply::Handled(); })
+					[
+						SNew(STextBlock)
+						.Justification(ETextJustify::Left)
+						.MinDesiredWidth(60)
+						.Text(FText::FromString(*this->SelectValues->operator[](i)))
+					]
+				]
+			];
+	}
+
+	return Contents;
+}
+
+void SMultiSelectComboBox::OnSelectedChangedFromMultiselect(TSharedPtr<FString> Value)
+{
+	if (Value->IsEmpty())
+	{
+		return;
+	}
+
+	if(!CheckedValues.Contains(*Value))
+	{
+		CheckedValues.Add(FString(*Value));
+	}
+	else
+	{
+		while (CheckedValues.Contains(*Value))
+		{
+			CheckedValues.Remove(FString(*Value));
+		}
+	}
+}
+
+
+#undef LOCTEXT_NAMESPACE
\ No newline at end of file
diff --git a/Engine/Source/Editor/StorageServerWidgets/Private/SMultiSelectComboBox.h b/Engine/Source/Editor/StorageServerWidgets/Private/SMultiSelectComboBox.h
new file mode 100644
--- /dev/null
+++ b/Engine/Source/Editor/StorageServerWidgets/Private/SMultiSelectComboBox.h
@@ -0,0 +1,32 @@
+// Copyright Epic Games, Inc. All Rights Reserved.
+
+#pragma once
+
+#include "Widgets/SCompoundWidget.h"
+#include "Styling/SlateTypes.h"
+#include "Containers/ObservableArray.h"
+
+DECLARE_DELEGATE_OneParam(FOnCheckboxChangedSignature, ECheckBoxState)
+
+class SMultiSelectComboBox : public SCompoundWidget
+{
+public:
+	SLATE_BEGIN_ARGS(SMultiSelectComboBox) {}
+		SLATE_ARGUMENT(TArray<TSharedPtr<FString>>*, SelectValues)
+		SLATE_EVENT(FOnCheckboxChangedSignature, OnCheckboxChanged)
+	SLATE_END_ARGS()
+
+	virtual ~SMultiSelectComboBox();
+
+	void Construct(const FArguments& InArgs);
+
+private:
+	TSharedRef<class SWidget> GenerateMenuItems();
+
+	void OnSelectedChangedFromMultiselect(TSharedPtr<FString> Value);
+
+	FReply OnLabelClick();
+
+	TArray<FString> CheckedValues;
+	TArray<TSharedPtr<FString>>* SelectValues;
+};
\ No newline at end of file
diff --git a/Engine/Source/Programs/BuildStorageTool/BuildStorageTool.Build.cs b/Engine/Source/Programs/BuildStorageTool/BuildStorageTool.Build.cs
new file mode 100644
--- /dev/null
+++ b/Engine/Source/Programs/BuildStorageTool/BuildStorageTool.Build.cs
@@ -0,0 +1,38 @@
+// Copyright Epic Games, Inc. All Rights Reserved.
+
+using UnrealBuildTool;
+
+public class BuildStorageTool : ModuleRules
+{
+	public BuildStorageTool(ReadOnlyTargetRules Target) : base(Target)
+	{
+		bTreatAsEngineModule = true;
+		
+		PrivateDependencyModuleNames.AddRange(
+			new string[]
+			{
+				"ApplicationCore",
+				"Core",
+				"DesktopPlatform",
+				"Json",
+				"Projects",
+				"Slate",
+				"SlateCore",
+				"StandaloneRenderer",
+				"StorageServerWidgets",
+				"Zen"
+			});
+
+		if (Target.IsInPlatformGroup(UnrealPlatformGroup.Linux))
+		{
+			PrivateDependencyModuleNames.AddRange(
+				new string[]
+				{
+					"UnixCommonStartup"
+				}
+			);
+		}
+
+		PublicIncludePathModuleNames.Add("Launch");
+	}
+}
diff --git a/Engine/Source/Programs/BuildStorageTool/BuildStorageTool.Target.cs b/Engine/Source/Programs/BuildStorageTool/BuildStorageTool.Target.cs
new file mode 100644
--- /dev/null
+++ b/Engine/Source/Programs/BuildStorageTool/BuildStorageTool.Target.cs
@@ -0,0 +1,44 @@
+// Copyright Epic Games, Inc. All Rights Reserved.
+
+using UnrealBuildTool;
+using System.Collections.Generic;
+
+[SupportedPlatforms(UnrealPlatformClass.Desktop)]
+public class BuildStorageToolTarget : TargetRules
+{
+	public BuildStorageToolTarget(TargetInfo Target) : base(Target)
+	{
+		Type = TargetType.Program;
+		LinkType = TargetLinkType.Monolithic;
+		IncludeOrderVersion = EngineIncludeOrderVersion.Latest;
+		LaunchModuleName = "BuildStorageTool";
+
+		bUseXGEController				= false;
+		bLoggingToMemoryEnabled			= true;
+		bUseLoggingInShipping			= true;
+		bCompileWithAccessibilitySupport= false;
+		bWithServerCode					= false;
+		bCompileNavmeshClusterLinks		= false;
+		bCompileNavmeshSegmentLinks		= false;
+		bCompileRecast					= false;
+		bCompileICU 					= false;
+		bWithLiveCoding					= false;
+		bBuildDeveloperTools			= false;
+		bBuildWithEditorOnlyData		= false;
+		bCompileAgainstEngine			= false;
+		bCompileAgainstCoreUObject		= true;
+		bCompileAgainstApplicationCore	= true;
+		bUsesSlate						= true;
+		bForceDisableAutomationTests	= true;
+
+		bIsBuildingConsoleApplication = false;
+
+		WindowsPlatform.bUseBundledDbgHelp = false;
+		WindowsPlatform.bPixProfilingEnabled = false;
+		WindowsPlatform.bMergeIdenticalCOMDATs = true;
+		WindowsPlatform.bStripUnreferencedSymbols = true;
+
+		GlobalDefinitions.Add("WITH_AUTOMATION_TESTS=0");
+		GlobalDefinitions.Add("WITH_AUTOMATION_WORKER=0");
+	}
+}
diff --git a/Engine/Source/Programs/BuildStorageTool/Private/BuildStorageTool.cpp b/Engine/Source/Programs/BuildStorageTool/Private/BuildStorageTool.cpp
new file mode 100644
--- /dev/null
+++ b/Engine/Source/Programs/BuildStorageTool/Private/BuildStorageTool.cpp
@@ -0,0 +1,349 @@
+// Copyright Epic Games, Inc. All Rights Reserved.
+
+#include "BuildStorageTool.h"
+#include "Experimental/ZenServerInterface.h"
+#include "Features/IModularFeatures.h"
+#include "Framework/Application/SlateApplication.h"
+#include "Framework/MultiBox/MultiBoxBuilder.h"
+#include "Framework/Notifications/NotificationManager.h"
+#include "GenericPlatform/GenericPlatformFile.h"
+#include "HAL/PlatformProcess.h"
+#include "Misc/CompilationResult.h"
+#include "Misc/MessageDialog.h"
+#include "Misc/Optional.h"
+#include "Misc/Timespan.h"
+#include "Modules/ModuleInterface.h"
+#include "Modules/ModuleManager.h"
+#include "RequiredProgramMainCPPInclude.h"
+#include "SBuildActivity.h"
+#include "SBuildSelection.h"
+#include "SBuildLogin.h"
+#include "SMessageDialog.h"
+#include "StandaloneRenderer.h"
+#include "BuildStorageToolStyle.h"
+#include "Tasks/Task.h"
+#include "Templates/SharedPointer.h"
+#include "Widgets/Input/SButton.h"
+#include "Widgets/Input/SButton.h"
+#include "Widgets/Input/SCheckBox.h"
+#include "Widgets/Notifications/SNotificationList.h"
+#include "Widgets/Text/STextBlock.h"
+#include "ZenServiceInstanceManager.h"
+#include "BuildServiceInstanceManager.h"
+
+#if PLATFORM_WINDOWS
+#include "Runtime/Launch/Resources/Windows/Resource.h"
+#include "Windows/WindowsApplication.h"
+#include "Windows/AllowWindowsPlatformTypes.h"
+#include <shellapi.h>
+#include "Windows/HideWindowsPlatformTypes.h"
+#elif PLATFORM_LINUX
+#include "UnixCommonStartup.h"
+#elif PLATFORM_MAC
+#include "Mac/MacProgramDelegate.h"
+#include "LaunchEngineLoop.h"
+#else
+#error "Unsupported platform!"
+#endif
+
+#define LOCTEXT_NAMESPACE "BuildStorageTool"
+
+// These macros are not properly defined by UBT in the case of an engine program with bTreatAsEngineModule=true
+// So define them here as a workaround
+#define IMPLEMENT_ENCRYPTION_KEY_REGISTRATION()
+#define IMPLEMENT_SIGNING_KEY_REGISTRATION()
+IMPLEMENT_APPLICATION(BuildStorageTool, "BuildStorageTool");
+
+static void OnRequestExit()
+{
+	RequestEngineExit(TEXT("BuildStorageTool Closed"));
+}
+
+static void HideOnCloseOverride(const TSharedRef<SWindow>& WindowBeingClosed)
+{
+	WindowBeingClosed->HideWindow();
+}
+
+class FBuildStorageToolApp
+{
+	FCriticalSection CriticalSection;
+	FSlateApplication& Slate;
+	TSharedPtr<SWindow> Window;
+	TSharedPtr<SNotificationItem> CompileNotification;
+	TSharedPtr<UE::Zen::FServiceInstanceManager> ZenServiceInstanceManager;
+	TSharedPtr<UE::Zen::Build::FServiceInstanceManager> BuildServiceInstanceManager;
+
+	std::atomic<bool> bLatentExclusiveOperationActive = false;
+
+	void ExitTool()
+	{
+		FSlateApplication::Get().RequestDestroyWindow(Window.ToSharedRef());
+	}
+
+	bool CanExecuteExclusiveAction()
+	{
+		return !bLatentExclusiveOperationActive;
+	}
+
+	TSharedRef< SWidget > MakeMainMenu()
+	{
+		FMenuBarBuilder MenuBuilder( nullptr );
+		{
+			// Control
+			MenuBuilder.AddPullDownMenu(
+				LOCTEXT( "FileMenu", "File" ),
+				LOCTEXT( "FileMenu_ToolTip", "Opens the file menu" ),
+				FOnGetContent::CreateRaw( this, &FBuildStorageToolApp::FillFileMenu ) );
+		}
+
+		// Create the menu bar
+		TSharedRef< SWidget > MenuBarWidget = MenuBuilder.MakeWidget();
+		MenuBarWidget->SetVisibility( EVisibility::Visible ); // Work around for menu bar not showing on Mac
+
+		return MenuBarWidget;
+	}
+
+	TSharedRef<SWidget> FillFileMenu()
+	{
+		const bool bCloseSelfOnly = false;
+		const bool bSearchable = false;
+		const bool bRecursivelySearchable = false;
+
+		FMenuBuilder MenuBuilder(true,
+			nullptr,
+			TSharedPtr<FExtender>(),
+			bCloseSelfOnly,
+			&FCoreStyle::Get(),
+			bSearchable,
+			NAME_None,
+			bRecursivelySearchable);
+
+		MenuBuilder.AddMenuEntry(
+			LOCTEXT("Exit", "Exit"),
+			LOCTEXT("Exit_ToolTip", "Exits the Storage Server UI"),
+			FSlateIcon(),
+			FUIAction(
+				FExecuteAction::CreateRaw( this, &FBuildStorageToolApp::ExitTool ),
+				FCanExecuteAction::CreateRaw(this, &FBuildStorageToolApp::CanExecuteExclusiveAction)
+			),
+			NAME_None,
+			EUserInterfaceActionType::Button
+		);
+
+		return MenuBuilder.MakeWidget();
+	}
+
+public:
+	FBuildStorageToolApp(FSlateApplication& InSlate)
+		: Slate(InSlate)
+	{
+		ZenServiceInstanceManager = MakeShared<UE::Zen::FServiceInstanceManager>();
+		BuildServiceInstanceManager = MakeShared<UE::Zen::Build::FServiceInstanceManager>();
+		InstallMessageDialogOverride();
+	}
+
+	virtual ~FBuildStorageToolApp()
+	{
+		RemoveMessageDialogOverride();
+	}
+
+	void Run()
+	{
+		Window =
+			SNew(SWindow)
+			.Title(GetWindowTitle())
+			.ClientSize(FVector2D(700.0f, 400.0f))
+			.ActivationPolicy(EWindowActivationPolicy::Always)
+			.SizingRule(ESizingRule::UserSized)
+			.IsTopmostWindow(false)
+			.FocusWhenFirstShown(false)
+			.SupportsMaximize(true)
+			.SupportsMinimize(true)
+			.HasCloseButton(true)
+			[
+				SNew(SBorder)
+				.BorderImage(FCoreStyle::Get().GetBrush("ToolPanel.GroupBorder"))
+				[
+					SNew(SVerticalBox)
+
+					// Menu
+					+ SVerticalBox::Slot()
+					.AutoHeight()
+					[
+						MakeMainMenu()
+					]
+
+					// Login panel
+					+ SVerticalBox::Slot()
+					.AutoHeight()
+					.Padding(0.0f, 10.0f, 0.0f, 0.0f)
+					.HAlign(HAlign_Fill)
+					.VAlign(VAlign_Top)
+					[
+						SNew(SBuildLogin)
+						.ZenServiceInstance(ZenServiceInstanceManager.ToSharedRef(), &UE::Zen::FServiceInstanceManager::GetZenServiceInstance)
+						.BuildServiceInstance(BuildServiceInstanceManager.ToSharedRef(), &UE::Zen::Build::FServiceInstanceManager::GetBuildServiceInstance)
+					]
+
+					// Selection panel
+					+ SVerticalBox::Slot()
+					.Padding(0.0f, 10.0f, 0.0f, 0.0f)
+					.HAlign(HAlign_Fill)
+					.VAlign(VAlign_Top)
+					[
+						SNew(SBuildSelection)
+						.ZenServiceInstance(ZenServiceInstanceManager.ToSharedRef(), &UE::Zen::FServiceInstanceManager::GetZenServiceInstance)
+						.BuildServiceInstance(BuildServiceInstanceManager.ToSharedRef(), &UE::Zen::Build::FServiceInstanceManager::GetBuildServiceInstance)
+					]
+
+					// Activity panel
+					+ SVerticalBox::Slot()
+					.AutoHeight()
+					.Padding(0.0f, 10.0f, 0.0f, 0.0f)
+					.HAlign(HAlign_Fill)
+					.VAlign(VAlign_Bottom)
+					[
+						SNew(SBuildActivity)
+						.ZenServiceInstance(ZenServiceInstanceManager.ToSharedRef(), &UE::Zen::FServiceInstanceManager::GetZenServiceInstance)
+						.BuildServiceInstance(BuildServiceInstanceManager.ToSharedRef(), &UE::Zen::Build::FServiceInstanceManager::GetBuildServiceInstance)
+					]
+				]
+			];
+
+		Slate.AddWindow(Window.ToSharedRef(), true);
+
+		// Setting focus seems to have to happen after the Window has been added
+		Slate.ClearKeyboardFocus(EFocusCause::Cleared);
+
+		// loop until the app is ready to quit
+		while (!IsEngineExitRequested())
+		{
+			BeginExitIfRequested();
+
+			FTaskGraphInterface::Get().ProcessThreadUntilIdle(ENamedThreads::GameThread);
+			FTSTicker::GetCoreTicker().Tick(FApp::GetDeltaTime());
+			Slate.PumpMessages();
+			Slate.Tick();
+
+			FPlatformProcess::Sleep(1.0f / 30.0f);
+		}
+
+		// Make sure the window is hidden, because it might take a while for the background thread to finish.
+		Window->HideWindow();
+	}
+
+private:
+	FText GetWindowTitle()
+	{
+		return LOCTEXT("WindowTitle", "Unreal Build Storage Tool");
+	}
+
+	EAppReturnType::Type OnModalMessageDialog(EAppMsgCategory InMessageCategory, EAppMsgType::Type InMessage, const FText& InText, const FText& InTitle)
+	{
+		if (IsInGameThread() && FSlateApplication::IsInitialized() && FSlateApplication::Get().CanAddModalWindow())
+		{
+			return OpenModalMessageDialog_Internal(InMessageCategory, InMessage, InText, InTitle, Window);
+		}
+		else
+		{
+			return FPlatformMisc::MessageBoxExt(InMessage, *InText.ToString(), *InTitle.ToString());
+		}
+	}
+
+	void InstallMessageDialogOverride()
+	{
+		FCoreDelegates::ModalMessageDialog.BindRaw(this, &FBuildStorageToolApp::OnModalMessageDialog);
+	}
+
+	void RemoveMessageDialogOverride()
+	{
+		FCoreDelegates::ModalMessageDialog.Unbind();
+	}
+};
+
+int BuildStorageToolMain(const TCHAR* CmdLine)
+{
+	ON_SCOPE_EXIT
+	{
+		RequestEngineExit(TEXT("Exiting"));
+		FEngineLoop::AppPreExit();
+		FModuleManager::Get().UnloadModulesAtShutdown();
+		FEngineLoop::AppExit();
+	};
+
+	// start up the main loop
+	GEngineLoop.PreInit(CmdLine);
+
+	FSystemWideCriticalSection SystemWideBuildStorageToolCritSec(TEXT("BuildStorageTool"));
+	if (!SystemWideBuildStorageToolCritSec.IsValid())
+	{
+		return true;
+	}
+	check(GConfig && GConfig->IsReadyForUse());
+
+	// Initialize high DPI mode
+	FSlateApplication::InitHighDPI(true);
+
+	{
+		// Create the platform slate application (what FSlateApplication::Get() returns)
+		TSharedRef<FSlateApplication> Slate = FSlateApplication::Create(MakeShareable(FPlatformApplicationMisc::CreateApplication()));
+
+		{
+			// Initialize renderer
+			TSharedRef<FSlateRenderer> SlateRenderer = GetStandardStandaloneRenderer();
+
+			// Try to initialize the renderer. It's possible that we launched when the driver crashed so try a few times before giving up.
+			bool bRendererInitialized = Slate->InitializeRenderer(SlateRenderer, true);
+			if (!bRendererInitialized)
+			{
+				// Close down the Slate application
+				FSlateApplication::Shutdown();
+				return false;
+			}
+
+			// Set the normal UE IsEngineExitRequested() when outer frame is closed
+			Slate->SetExitRequestedHandler(FSimpleDelegate::CreateStatic(&OnRequestExit));
+
+			// Prepare the custom Slate styles
+			FBuildStorageToolStyle::Initialize();
+
+			// Set the icon
+			FAppStyle::SetAppStyleSet(FBuildStorageToolStyle::Get());
+
+			// Run the inner application loop
+			FBuildStorageToolApp App(Slate.Get());
+			App.Run();
+
+			// Clean up the custom styles
+			FBuildStorageToolStyle::Shutdown();
+		}
+
+		// Close down the Slate application
+		FSlateApplication::Shutdown();
+	}
+
+	return true;
+}
+
+#if PLATFORM_WINDOWS
+int WINAPI WinMain(_In_ HINSTANCE hCurrInstance, _In_opt_ HINSTANCE hPrevInstance, _In_ LPSTR lpCmdLine, _In_ int nShowCmd)
+{
+	hInstance = hCurrInstance;
+	return BuildStorageToolMain(GetCommandLineW())? 0 : 1;
+}
+
+#elif PLATFORM_LINUX
+int main(int argc, char* argv[])
+{
+	return CommonUnixMain(argc, argv, &BuildStorageToolMain);
+}
+#elif PLATFORM_MAC
+int main(int argc, char* argv[])
+{
+	[MacProgramDelegate mainWithArgc : argc argv : argv programMain : BuildStorageToolMain programExit : FEngineLoop::AppExit] ;
+}
+#else
+#error "Unsupported platform!"
+#endif
+
+#undef LOCTEXT_NAMESPACE
diff --git a/Engine/Source/Programs/BuildStorageTool/Private/BuildStorageTool.h b/Engine/Source/Programs/BuildStorageTool/Private/BuildStorageTool.h
new file mode 100644
--- /dev/null
+++ b/Engine/Source/Programs/BuildStorageTool/Private/BuildStorageTool.h
@@ -0,0 +1,3 @@
+// Copyright Epic Games, Inc. All Rights Reserved.
+
+#pragma once
diff --git a/Engine/Source/Programs/BuildStorageTool/Private/BuildStorageToolStyle.cpp b/Engine/Source/Programs/BuildStorageTool/Private/BuildStorageToolStyle.cpp
new file mode 100644
--- /dev/null
+++ b/Engine/Source/Programs/BuildStorageTool/Private/BuildStorageToolStyle.cpp
@@ -0,0 +1,116 @@
+// Copyright Epic Games, Inc. All Rights Reserved.
+
+#include "BuildStorageToolStyle.h"
+#include "Misc/AssertionMacros.h"
+#include "Containers/UnrealString.h"
+#include "Misc/Paths.h"
+#include "Math/Color.h"
+#include "Math/Vector2D.h"
+#include "Styling/SlateStyleMacros.h"
+#include "Styling/SlateStyleRegistry.h"
+#include "Styling/SlateTypes.h"
+#include "Styling/CoreStyle.h"
+
+#define EDITOR_IMAGE_BRUSH(RelativePath, ...) IMAGE_BRUSH("../Editor/Slate/" RelativePath, __VA_ARGS__)
+#define EDITOR_IMAGE_BRUSH_SVG(RelativePath, ...) IMAGE_BRUSH_SVG("../Editor/Slate/" RelativePath, __VA_ARGS__)
+#define EDITOR_BOX_BRUSH(RelativePath, ...) BOX_BRUSH("../Editor/Slate/" RelativePath, __VA_ARGS__)
+#define EDITOR_BORDER_BRUSH(RelativePath, ...) BORDER_BRUSH("../Editor/Slate/" RelativePath, __VA_ARGS__)
+#define RootToContentDir ContentFromEngine
+
+TSharedPtr< FSlateStyleSet > FBuildStorageToolStyle::StyleSet = nullptr;
+
+void FBuildStorageToolStyle::Initialize()
+{
+	if(!StyleSet.IsValid())
+	{
+		StyleSet = Create();
+		FSlateStyleRegistry::RegisterSlateStyle(*StyleSet);
+	}
+}
+
+void FBuildStorageToolStyle::Shutdown()
+{
+	FSlateStyleRegistry::UnRegisterSlateStyle(*StyleSet);
+	ensure(StyleSet.IsUnique());
+	StyleSet.Reset();
+}
+
+namespace
+{
+	FString ContentFromEngine(const FString& RelativePath, const TCHAR* Extension)
+	{
+		static const FString ContentDir = FPaths::EngineDir() / TEXT("Content/Slate");
+		return ContentDir / RelativePath + Extension;
+	}
+}
+
+TSharedRef< FSlateStyleSet > FBuildStorageToolStyle::Create()
+{
+	TSharedRef<FSlateStyleSet> StyleRef = MakeShareable(new FSlateStyleSet("BuildStorageToolStyle"));
+	FSlateStyleSet& Style = StyleRef.Get();
+
+	Style.SetParentStyleName("CoreStyle");
+
+	const FVector2D Icon12x12(12.0f, 12.0f);
+	const FVector2D Icon16x16(16.0f, 16.0f);
+
+	Style.Set( "AppIcon", new IMAGE_BRUSH( "Icons/DefaultAppIcon", FVector2D(20,20)) );
+
+	Style.Set("Zen.Icons.FolderExplore", new EDITOR_IMAGE_BRUSH_SVG("Starship/Common/ContentBrowser", Icon12x12));
+	Style.Set("Zen.Icons.Clipboard", new EDITOR_IMAGE_BRUSH_SVG("Starship/Common/Clipboard", Icon16x16));
+
+	const FTextBlockStyle DefaultText = FTextBlockStyle()
+		.SetFont(DEFAULT_FONT("Bold", 10))
+		.SetColorAndOpacity(FSlateColor::UseForeground())
+		.SetShadowOffset(FVector2D::ZeroVector)
+		.SetShadowColorAndOpacity(FLinearColor::Black);
+
+	// Set the client app styles
+	Style.Set(TEXT("Code"), FTextBlockStyle(DefaultText)
+		.SetFont(DEFAULT_FONT("Regular", 8))
+		.SetColorAndOpacity(FSlateColor(FLinearColor::White * 0.8f))
+	);
+
+	Style.Set(TEXT("Title"), FTextBlockStyle(DefaultText)
+		.SetFont(DEFAULT_FONT("Bold", 12))
+	);
+
+	Style.Set(TEXT("Status"), FTextBlockStyle(DefaultText)
+		.SetColorAndOpacity(FSlateColor::UseSubduedForeground())
+	);
+
+
+	FSlateBrush* GenericWhiteBox = new IMAGE_BRUSH( "Old/White", Icon16x16 );
+
+	// Scrollbar
+	const FScrollBarStyle ScrollBar = FScrollBarStyle()
+		.SetVerticalTopSlotImage(IMAGE_BRUSH("Common/Scrollbar_Background_Vertical", FVector2D(8, 8)))
+		.SetVerticalBottomSlotImage(IMAGE_BRUSH("Common/Scrollbar_Background_Vertical", FVector2D(8, 8)))
+		.SetHorizontalTopSlotImage(IMAGE_BRUSH("Common/Scrollbar_Background_Horizontal", FVector2D(8, 8)))
+		.SetHorizontalBottomSlotImage(IMAGE_BRUSH("Common/Scrollbar_Background_Horizontal", FVector2D(8, 8)))
+		.SetNormalThumbImage(BOX_BRUSH("Common/Scrollbar_Thumb", FMargin(4.f / 16.f)))
+		.SetDraggedThumbImage(BOX_BRUSH("Common/Scrollbar_Thumb", FMargin(4.f / 16.f)))
+		.SetHoveredThumbImage(BOX_BRUSH("Common/Scrollbar_Thumb", FMargin(4.f / 16.f)));
+
+	Style.Set("Log.TextBox", FEditableTextBoxStyle(FCoreStyle::Get().GetWidgetStyle<FEditableTextBoxStyle>("NormalEditableTextBox"))
+		.SetBackgroundImageNormal( BOX_BRUSH("Common/WhiteGroupBorder", FMargin(4.0f/16.0f)))
+		.SetBackgroundImageHovered( BOX_BRUSH("Common/WhiteGroupBorder", FMargin(4.0f/16.0f)))
+		.SetBackgroundImageFocused(BOX_BRUSH("Common/WhiteGroupBorder", FMargin(4.0f / 16.0f)))
+		.SetBackgroundImageReadOnly(BOX_BRUSH("Common/WhiteGroupBorder", FMargin(4.0f / 16.0f)))
+		.SetBackgroundColor( FLinearColor(0.015f, 0.015f, 0.015f) )
+		.SetScrollBarStyle(ScrollBar)
+		);
+
+	return StyleRef;
+}
+
+const ISlateStyle& FBuildStorageToolStyle::Get()
+{
+	return *StyleSet;
+}
+
+#undef EDITOR_IMAGE_BRUSH
+#undef EDITOR_IMAGE_BRUSH_SVG
+#undef EDITOR_BOX_BRUSH
+#undef EDITOR_BORDER_BRUSH
+#undef RootToContentDir
diff --git a/Engine/Source/Programs/BuildStorageTool/Private/BuildStorageToolStyle.h b/Engine/Source/Programs/BuildStorageTool/Private/BuildStorageToolStyle.h
new file mode 100644
--- /dev/null
+++ b/Engine/Source/Programs/BuildStorageTool/Private/BuildStorageToolStyle.h
@@ -0,0 +1,31 @@
+// Copyright Epic Games, Inc. All Rights Reserved.
+
+#pragma once
+
+#include "CoreMinimal.h"
+#include "Styling/SlateStyle.h"
+
+class FBuildStorageToolStyle
+{
+public:
+	/**
+	 * Set up specific styles for the storage server UI app
+	 */
+	static void Initialize();
+
+	/**
+	 * Tidy up on shut-down
+	 */
+	static void Shutdown();
+
+	/*
+	 * Access to singleton style object
+	 */ 
+	static const ISlateStyle& Get();
+
+private:
+	static TSharedRef<FSlateStyleSet> Create();
+
+	/** Singleton style object */
+	static TSharedPtr<FSlateStyleSet> StyleSet;
+};
diff --git a/Engine/Source/Programs/BuildStorageTool/Private/SMessageDialog.cpp b/Engine/Source/Programs/BuildStorageTool/Private/SMessageDialog.cpp
new file mode 100644
--- /dev/null
+++ b/Engine/Source/Programs/BuildStorageTool/Private/SMessageDialog.cpp
@@ -0,0 +1,447 @@
+// Copyright Epic Games, Inc. All Rights Reserved.
+
+
+#include "SMessageDialog.h"
+#include "Fonts/SlateFontInfo.h"
+#include "Framework/Application/SlateApplication.h"
+#include "HAL/PlatformApplicationMisc.h"
+#include "Input/Reply.h"
+#include "InputCoreTypes.h"
+#include "Misc/Attribute.h"
+#include "SlateOptMacros.h"
+#include "Styling/CoreStyle.h"
+#include "Widgets/DeclarativeSyntaxSupport.h"
+#include "Widgets/Input/SButton.h"
+#include "Widgets/Images/SImage.h"
+#include "Widgets/Layout/SBorder.h"
+#include "Widgets/Layout/SScrollBox.h"
+#include "Widgets/Layout/SUniformGridPanel.h"
+#include "Widgets/Input/SCheckBox.h"
+#include "Widgets/SBoxPanel.h"
+#include "Widgets/SCompoundWidget.h"
+#include "Widgets/SWindow.h"
+#include "Widgets/Text/STextBlock.h"
+
+#define LOCTEXT_NAMESPACE "Dialogs"
+
+class SChoiceDialog : public SCompoundWidget
+{
+public:
+	SLATE_BEGIN_ARGS( SChoiceDialog )	{}
+		SLATE_ARGUMENT(TSharedPtr<SWindow>, ParentWindow)
+		SLATE_ATTRIBUTE(FText, Message)	
+		SLATE_ATTRIBUTE(float, WrapMessageAt)
+		SLATE_ARGUMENT(EAppMsgCategory, MessageCategory)
+		SLATE_ARGUMENT(EAppMsgType::Type, MessageType)
+	SLATE_END_ARGS()
+
+	BEGIN_SLATE_FUNCTION_BUILD_OPTIMIZATION
+	void Construct( const FArguments& InArgs )
+	{
+		ParentWindow = InArgs._ParentWindow;
+		ParentWindow->SetWidgetToFocusOnActivate(SharedThis(this));
+		Response = EAppReturnType::Cancel;
+		MessageType = InArgs._MessageType;
+
+		FSlateFontInfo MessageFont( FAppStyle::GetFontStyle("StandardDialog.LargeFont"));
+		MyMessage = InArgs._Message;
+
+		TSharedPtr<SUniformGridPanel> ButtonBox;
+
+		const FSlateBrush* IconBrush = FAppStyle::Get().GetBrush("Icons.WarningWithColor.Large");
+		switch (InArgs._MessageCategory)
+		{
+		case EAppMsgCategory::Error:
+			IconBrush = FAppStyle::Get().GetBrush("Icons.ErrorWithColor.Large");
+			break;
+		case EAppMsgCategory::Success:
+			IconBrush = FAppStyle::Get().GetBrush("Icons.SuccessWithColor.Large");
+			break;
+		case EAppMsgCategory::Info:
+			IconBrush = FAppStyle::Get().GetBrush("Icons.InfoWithColor.Large");
+			break;
+		default:
+			break;
+		}
+
+		this->ChildSlot
+		[	
+			SNew(SBorder)
+			.BorderImage(FAppStyle::GetBrush("ToolPanel.GroupBorder"))
+			.Padding(16.f)
+			[
+				SNew(SVerticalBox)
+
+				+ SVerticalBox::Slot()
+				.HAlign(HAlign_Fill)
+				.VAlign(VAlign_Fill)
+				.FillHeight(1.0f)
+				.MaxHeight(550)
+				[
+					SNew(SHorizontalBox)
+					+SHorizontalBox::Slot()
+					.AutoWidth()
+					.VAlign(VAlign_Top)
+					.HAlign(HAlign_Left)
+					[
+						SNew(SImage)
+						.DesiredSizeOverride(FVector2D(24.f, 24.f))
+						.Image(IconBrush)
+					]
+
+					+SHorizontalBox::Slot()
+					.Padding(16.f, 0.f, 0.f, 0.f)
+					[
+						SNew(SScrollBox)
+						+ SScrollBox::Slot()
+						.HAlign(HAlign_Left)
+						.VAlign(VAlign_Center)
+						[
+							SNew(STextBlock)
+							.Text(MyMessage)
+							.Font(MessageFont)
+							.WrapTextAt(InArgs._WrapMessageAt)
+						]
+					]
+				]
+
+				+SVerticalBox::Slot()
+				.Padding(FMargin(0.f, 32.f, 0.f, 0.f))
+				.AutoHeight()
+				[
+					SNew(SHorizontalBox)
+
+					+ SHorizontalBox::Slot()
+					.AutoWidth()
+					.HAlign(HAlign_Left)
+					.VAlign(VAlign_Center)
+					[
+						SNew(SButton)
+						.ButtonStyle(FAppStyle::Get(), "SimpleButton")
+						.OnClicked(this, &SChoiceDialog::HandleCopyMessageButtonClicked)
+						.ToolTipText(NSLOCTEXT("SChoiceDialog", "CopyMessageTooltip", "Copy the text in this message to the clipboard (CTRL+C)"))
+						.ContentPadding(2.f)
+						.Content()
+						[
+							SNew(SImage)
+							.Image(FAppStyle::Get().GetBrush("Zen.Icons.Clipboard"))
+							.ColorAndOpacity(FSlateColor::UseForeground())
+						]
+					]
+
+					+ SHorizontalBox::Slot()
+						.FillWidth(1.0f)
+						.HAlign(HAlign_Left)
+						.VAlign(VAlign_Center)
+						.Padding(FMargin(16.f, 0.f, 0.f, 0.f))
+						[
+							SNew(SCheckBox)
+							.IsChecked(ECheckBoxState::Unchecked)
+							.OnCheckStateChanged(this, &SChoiceDialog::OnCheckboxClicked)
+							.Visibility(this, &SChoiceDialog::GetCheckboxVisibility)
+							.ToolTipText(NSLOCTEXT("SChoiceDialog", "ApplyToAllTooltip", "Make your choice of Yes or No apply to all remaining items in the current operation"))
+							[
+								SNew(STextBlock)
+								.WrapTextAt(615.0f)
+								.Text(NSLOCTEXT("SChoiceDialog", "ApplyToAllLabel", "Apply to All"))
+							]
+						]
+
+					+ SHorizontalBox::Slot()
+					.AutoWidth()
+					.HAlign(HAlign_Right)
+					.VAlign(VAlign_Center)
+					[
+						SAssignNew( ButtonBox, SUniformGridPanel )
+						.SlotPadding(FMargin(16.f, 0.f, 0.f, 0.f))
+						.MinDesiredSlotWidth(FAppStyle::Get().GetFloat("StandardDialog.MinDesiredSlotWidth"))
+						.MinDesiredSlotHeight(FAppStyle::Get().GetFloat("StandardDialog.MinDesiredSlotHeight"))
+					]
+				]
+			]
+		];
+
+		int32 SlotIndex = 0;
+
+#define ADD_SLOT(Button)\
+		ButtonBox->AddSlot(SlotIndex++,0)\
+		[\
+			SNew( SButton )\
+			.VAlign(VAlign_Center)\
+			.HAlign(HAlign_Center)\
+			.Text( EAppReturnTypeToText(EAppReturnType::Button) )\
+			.OnClicked( this, &SChoiceDialog::HandleButtonClicked, EAppReturnType::Button )\
+		];
+
+#define ADD_SLOT_PRIMARY(Button)\
+		ButtonBox->AddSlot(SlotIndex++,0)\
+		[\
+			SNew( SButton )\
+			.VAlign(VAlign_Center)\
+			.HAlign(HAlign_Center)\
+			.ButtonStyle(&FAppStyle::Get().GetWidgetStyle<FButtonStyle>("PrimaryButton"))\
+			.Text( EAppReturnTypeToText(EAppReturnType::Button) )\
+			.OnClicked( this, &SChoiceDialog::HandleButtonClicked, EAppReturnType::Button )\
+		];
+
+
+
+		switch ( MessageType )
+		{	
+		case EAppMsgType::Ok:
+			ADD_SLOT_PRIMARY(Ok)
+			break;
+		case EAppMsgType::YesNo:
+			ADD_SLOT_PRIMARY(Yes)
+			ADD_SLOT(No)
+			break;
+		case EAppMsgType::OkCancel:
+			ADD_SLOT_PRIMARY(Ok)
+			ADD_SLOT(Cancel)
+			break;
+		case EAppMsgType::YesNoCancel:
+			ADD_SLOT_PRIMARY(Yes)
+			ADD_SLOT(No)
+			ADD_SLOT(Cancel)
+			break;
+		case EAppMsgType::CancelRetryContinue:
+			ADD_SLOT_PRIMARY(Continue)
+			ADD_SLOT(Retry)
+			ADD_SLOT(Cancel)
+			break;
+		case EAppMsgType::YesNoYesAllNoAll:
+			ADD_SLOT_PRIMARY(Yes)
+			ADD_SLOT(No)
+			break;
+		case EAppMsgType::YesNoYesAllNoAllCancel:
+			ADD_SLOT_PRIMARY(Yes)
+			ADD_SLOT(No)
+			ADD_SLOT(Cancel)
+			break;
+		case EAppMsgType::YesNoYesAll:
+			ADD_SLOT_PRIMARY(Yes)
+			ADD_SLOT(No)
+			ADD_SLOT(YesAll)
+			break;
+		}
+
+#undef ADD_SLOT
+#undef ADD_SLOT_PRIMARY
+	}
+	END_SLATE_FUNCTION_BUILD_OPTIMIZATION
+
+
+	EAppReturnType::Type GetResponse()
+	{
+		return Response;
+	}
+
+	virtual	FReply OnKeyDown( const FGeometry& MyGeometry, const FKeyEvent& InKeyEvent )
+	{
+		if( InKeyEvent.GetKey() == EKeys::Escape )
+		{
+			return HandleButtonClicked(EAppReturnType::Cancel);
+		}
+
+		if (InKeyEvent.GetKey() == EKeys::C && InKeyEvent.IsControlDown())
+		{
+			CopyMessageToClipboard();
+
+			return FReply::Handled();
+		}
+
+		//if it was some other button, ignore it
+		return FReply::Unhandled();
+	}
+
+	/** Override the base method to allow for keyboard focus */
+	virtual bool SupportsKeyboardFocus() const
+	{
+		return true;
+	}
+
+	/** Converts an EAppReturnType into a localized FText */
+	static FText EAppReturnTypeToText(EAppReturnType::Type ReturnType)
+	{
+		switch(ReturnType)
+		{
+		case EAppReturnType::No:
+			return LOCTEXT("EAppReturnTypeNo", "No");
+		case EAppReturnType::Yes:
+			return LOCTEXT("EAppReturnTypeYes", "Yes");
+		case EAppReturnType::YesAll:
+			return LOCTEXT("EAppReturnTypeYesAll", "Yes All");
+		case EAppReturnType::NoAll:
+			return LOCTEXT("EAppReturnTypeNoAll", "No All");
+		case EAppReturnType::Cancel:
+			return LOCTEXT("EAppReturnTypeCancel", "Cancel");
+		case EAppReturnType::Ok:
+			return LOCTEXT("EAppReturnTypeOk", "OK");
+		case EAppReturnType::Retry:
+			return LOCTEXT("EAppReturnTypeRetry", "Retry");
+		case EAppReturnType::Continue:
+			return LOCTEXT("EAppReturnTypeContinue", "Continue");
+		default:
+			return LOCTEXT("MissingType", "MISSING RETURN TYPE");
+		}
+	}
+
+protected:
+
+	/**
+	 * Copies the message text to the clipboard.
+	 */
+	void CopyMessageToClipboard( )
+	{
+		FPlatformApplicationMisc::ClipboardCopy( *MyMessage.Get().ToString() );
+	}
+
+
+private:
+
+	// Handles clicking a message box button.
+	FReply HandleButtonClicked( EAppReturnType::Type InResponse )
+	{
+		Response = InResponse;
+		if ((MessageType == EAppMsgType::YesNoYesAllNoAll || MessageType == EAppMsgType::YesNoYesAllNoAllCancel)
+			&& bApplyToAllChecked)
+		{
+			if (Response == EAppReturnType::Yes)
+			{
+				Response = EAppReturnType::YesAll;
+			}
+			else if (Response == EAppReturnType::No)
+			{
+				Response = EAppReturnType::NoAll;
+			}
+		}
+
+		ResultCallback.ExecuteIfBound(ParentWindow.ToSharedRef(), Response);
+
+
+		ParentWindow->RequestDestroyWindow();
+
+		return FReply::Handled();
+	}
+
+	// Handles clicking the 'Copy Message' button.
+	FReply HandleCopyMessageButtonClicked()
+	{
+		CopyMessageToClipboard();
+		return FReply::Handled();
+	}
+		
+	// Used as a delegate for the OnClicked property of the Apply to All checkbox
+	void OnCheckboxClicked(ECheckBoxState InNewState)
+	{
+		bApplyToAllChecked = InNewState == ECheckBoxState::Checked;
+	}
+
+	// Used as a delegate for the Visibility property of the Apply to All checkbox
+	EVisibility GetCheckboxVisibility() const
+	{
+		if (MessageType == EAppMsgType::YesNoYesAllNoAll || MessageType == EAppMsgType::YesNoYesAllNoAllCancel)
+		{
+			return EVisibility::Visible;
+		}
+		return EVisibility::Hidden;
+	}
+
+public:
+	/** Callback delegate that is triggered, when the dialog is run in non-modal mode */
+	FOnMsgDlgResult ResultCallback;
+
+private:
+
+	EAppReturnType::Type Response;
+	TSharedPtr<SWindow> ParentWindow;
+	TAttribute<FText> MyMessage;
+	bool bApplyToAllChecked;
+	EAppMsgType::Type MessageType;
+};
+
+
+void CreateMsgDlgWindow(TSharedPtr<SWindow>& OutWindow, TSharedPtr<SChoiceDialog>& OutDialog, EAppMsgCategory InMessageCategory, EAppMsgType::Type InMessageType,
+						const FText& InMessage, const FText& InTitle, FOnMsgDlgResult ResultCallback=NULL)
+{
+	OutWindow = SNew(SWindow)
+		.Title(InTitle)
+		.SizingRule(ESizingRule::Autosized)
+		.AutoCenter(EAutoCenter::PreferredWorkArea)
+		.SupportsMinimize(false).SupportsMaximize(false);
+
+	OutDialog = SNew(SChoiceDialog)
+		.ParentWindow(OutWindow)
+		.Message(InMessage)
+		.WrapMessageAt(512.0f)
+		.MessageCategory(InMessageCategory)
+		.MessageType(InMessageType);
+
+	OutDialog->ResultCallback = ResultCallback;
+
+	OutWindow->SetContent(OutDialog.ToSharedRef());
+}
+
+EAppReturnType::Type OpenModalMessageDialog_Internal(EAppMsgCategory InMessageCategory, EAppMsgType::Type InMessageType, EAppReturnType::Type InDefaultValue, const FText& InMessage, const FText& InTitle, const TSharedPtr<const SWidget>& ModalParent)
+{
+	EAppReturnType::Type Response = InDefaultValue;
+	TSharedPtr<SWindow> MsgWindow = NULL;
+	TSharedPtr<SChoiceDialog> MsgDialog = NULL;
+
+	CreateMsgDlgWindow(MsgWindow, MsgDialog, InMessageCategory, InMessageType, InMessage, InTitle);
+
+	FSlateApplication::Get().AddModalWindow(MsgWindow.ToSharedRef(), ModalParent);
+
+	Response = MsgDialog->GetResponse();
+
+	return Response;
+}
+
+EAppReturnType::Type OpenModalMessageDialog_Internal(EAppMsgCategory InMessageCategory, EAppMsgType::Type InMessageType, const FText& InMessage, const FText& InTitle, const TSharedPtr<const SWidget>& ModalParent)
+{
+	EAppReturnType::Type DefaultValue = EAppReturnType::Yes;
+	switch (InMessageType)
+	{
+	case EAppMsgType::Ok:
+		DefaultValue = EAppReturnType::Ok;
+		break;
+	case EAppMsgType::YesNo:
+		DefaultValue = EAppReturnType::No;
+		break;
+	case EAppMsgType::OkCancel:
+		DefaultValue = EAppReturnType::Cancel;
+		break;
+	case EAppMsgType::YesNoCancel:
+		DefaultValue = EAppReturnType::Cancel;
+		break;
+	case EAppMsgType::CancelRetryContinue:
+		DefaultValue = EAppReturnType::Cancel;
+		break;
+	case EAppMsgType::YesNoYesAllNoAll:
+		DefaultValue = EAppReturnType::No;
+		break;
+	case EAppMsgType::YesNoYesAllNoAllCancel:
+		DefaultValue = EAppReturnType::Cancel;
+		break;
+	default:
+		DefaultValue = EAppReturnType::Yes;
+		break;
+	}
+
+	if (GIsRunningUnattendedScript && InMessageType != EAppMsgType::Ok)
+	{
+		//UE_LOG(LogDialogs, Error, TEXT("Message Dialog was triggered in unattended script mode without a default value. %d will be used."), (int32)DefaultValue);
+		if (FPlatformMisc::IsDebuggerPresent())
+		{
+			UE_DEBUG_BREAK();
+		}
+		else
+		{
+			FDebug::DumpStackTraceToLog(ELogVerbosity::Error);
+		}
+	}
+
+	return OpenModalMessageDialog_Internal(InMessageCategory, InMessageType, DefaultValue, InMessage, InTitle, ModalParent);
+}
+
+#undef LOCTEXT_NAMESPACE
diff --git a/Engine/Source/Programs/BuildStorageTool/Private/SMessageDialog.h b/Engine/Source/Programs/BuildStorageTool/Private/SMessageDialog.h
new file mode 100644
--- /dev/null
+++ b/Engine/Source/Programs/BuildStorageTool/Private/SMessageDialog.h
@@ -0,0 +1,16 @@
+// Copyright Epic Games, Inc. All Rights Reserved.
+
+#pragma once
+
+#include "Delegates/DelegateCombinations.h"
+#include "GenericPlatform/GenericPlatformMisc.h"
+#include "Internationalization/Text.h"
+
+class SWindow;
+class SWidget;
+
+// Largely a copy of UnrealEd Dialogs.h
+
+DECLARE_DELEGATE_TwoParams(FOnMsgDlgResult, const TSharedRef<SWindow>&, EAppReturnType::Type);
+
+EAppReturnType::Type OpenModalMessageDialog_Internal(EAppMsgCategory InMessageCategory, EAppMsgType::Type InMessageType, const FText& InMessage, const FText& InTitle, const TSharedPtr<const SWidget>& ModalParent);
