BuildStorageTool - functional version 1.0.

diff --git a/Engine/Source/Developer/Zen/Private/BuildServerInterface.cpp b/Engine/Source/Developer/Zen/Private/BuildServerInterface.cpp
--- a/Engine/Source/Developer/Zen/Private/BuildServerInterface.cpp
+++ b/Engine/Source/Developer/Zen/Private/BuildServerInterface.cpp
@@ -138,6 +138,11 @@ FBuildServiceInstance::FBuildServiceInstance(FStringView InstanceURL)
 
 FBuildServiceInstance::~FBuildServiceInstance()
 {
+	bBuildTransferThreadStopping = true;
+	if (BuildTransferThread.IsJoinable())
+	{
+		BuildTransferThread.Join();
+	}
 }
 
 void
@@ -320,6 +325,72 @@ FBuildServiceInstance::ListBuilds(FStringView Namespace, FStringView Bucket, FOn
 	});
 }
 
+FBuildServiceInstance::FBuildTransfer
+FBuildServiceInstance::StartBuildTransfer(const FCbObjectId& BuildId, FStringView DestinationFolder, FStringView Namespace, FStringView Bucket)
+{
+	const FString ZenUtility = UE::Zen::GetLocalInstallUtilityPath();
+	TSharedPtr<FBuildTransferCommand> NewTransferCommand = MakeShared<FBuildTransferCommand>();
+	NewTransferCommand->Type = EBuildTransferType::Files;
+	NewTransferCommand->BuildId = BuildId;
+	NewTransferCommand->Namespace = Namespace;
+	NewTransferCommand->Bucket = Bucket;
+	NewTransferCommand->Destination = DestinationFolder;
+	BuildTransferThreadCommands.Enqueue(NewTransferCommand.ToSharedRef());
+	KickBuildTransferThread();
+
+	FBuildTransfer NewTransfer;
+	NewTransfer.BuildTransferCommand = NewTransferCommand;
+	return NewTransfer;
+}
+
+FBuildServiceInstance::FBuildTransfer
+FBuildServiceInstance::StartOplogBuildTransfer(const FCbObjectId& BuildId, FStringView DestinationProjectId, FStringView DestinationOplogId, FStringView ProjectFilePath, FStringView Namespace, FStringView Bucket)
+{
+	const FString ZenUtility = UE::Zen::GetLocalInstallUtilityPath();
+	TSharedPtr<FBuildTransferCommand> NewTransferCommand = MakeShared<FBuildTransferCommand>();
+	NewTransferCommand->Type = EBuildTransferType::Oplog;
+	NewTransferCommand->BuildId = BuildId;
+	NewTransferCommand->Namespace = Namespace;
+	NewTransferCommand->Bucket = Bucket;
+	NewTransferCommand->Destination = *WriteToString<64>(DestinationProjectId, TEXT("/"), DestinationOplogId);
+	NewTransferCommand->ProjectFilePath = ProjectFilePath;
+	BuildTransferThreadCommands.Enqueue(NewTransferCommand.ToSharedRef());
+	KickBuildTransferThread();
+
+	FBuildTransfer NewTransfer;
+	NewTransfer.BuildTransferCommand = NewTransferCommand;
+	return NewTransfer;
+}
+
+FString
+FBuildServiceInstance::FBuildTransfer::GetDescription() const
+{
+	if (!BuildTransferCommand)
+	{
+		return TEXT("null");
+	}
+
+	FReadScopeLock _(BuildTransferCommand->Lock);
+	int32 Index = 0;
+	if (BuildTransferCommand->Output.FindLastChar(TCHAR('\n'), Index))
+	{
+		return BuildTransferCommand->Output.RightChop(Index+1);
+	}
+	return BuildTransferCommand->Output;
+}
+
+FBuildServiceInstance::EBuildTransferStatus
+FBuildServiceInstance::FBuildTransfer::GetStatus() const
+{
+	if (!BuildTransferCommand)
+	{
+		return EBuildTransferStatus::Invalid;
+	}
+
+	FReadScopeLock _(BuildTransferCommand->Lock);
+	return BuildTransferCommand->Status;
+}
+
 FBuildServiceInstance::EConnectionState
 FBuildServiceInstance::GetConnectionState() const
 {
@@ -557,4 +628,129 @@ FBuildServiceInstance::CallOnRefreshNamespacesAndBucketsComplete(FOnRefreshNames
 	RefreshNamespacesAndBucketsComplete.Broadcast();
 }
 
+void
+FBuildServiceInstance::KickBuildTransferThread()
+{
+	if (!bBuildTransferThreadStarting.load(std::memory_order_relaxed) && !bBuildTransferThreadStarting.exchange(true, std::memory_order_relaxed))
+	{
+		BuildTransferThread = FThread(TEXT("BuildTransfer"), [this] { BuildTransferThreadLoop(); }, 128 * 1024);
+	}
+}
+
+void
+FBuildServiceInstance::BuildTransferThreadLoop()
+{
+	while (!BuildTransferThreadCommands.IsEmpty() || !bBuildTransferThreadStopping.load(std::memory_order_relaxed))
+	{
+		TOptional<TSharedRef<FBuildTransferCommand>> OptionalCommand = BuildTransferThreadCommands.Dequeue();
+		if (!OptionalCommand)
+		{
+			FPlatformProcess::Sleep(0.2f);
+			continue;
+		}
+		TSharedRef<FBuildTransferCommand> TransferCommand(OptionalCommand.GetValue());
+
+ 		FString ZenUtilityPath = FPaths::ConvertRelativePathToFull(UE::Zen::GetLocalInstallUtilityPath());
+ 		FString OidcTokenExecutableFilename = FPaths::ConvertRelativePathToFull(FDesktopPlatformModule::TryGet()->GetOidcTokenExecutableFilename(FPaths::RootDir()));
+
+		FString CommandLineArgs;
+		switch (TransferCommand->Type)
+		{
+		case EBuildTransferType::Files:
+			CommandLineArgs = FString::Printf(TEXT("builds download --host %s --local-path \"%s\" --namespace %s --bucket %s --build-id %s --oidctoken-exe-path \"%s\""),
+				*WriteToString<64>(EffectiveDomain), *FPaths::ConvertRelativePathToFull(TransferCommand->Destination), *TransferCommand->Namespace, *TransferCommand->Bucket,
+				*WriteToString<64>(TransferCommand->BuildId), *OidcTokenExecutableFilename);
+			break;
+		case EBuildTransferType::Oplog:
+			{
+				UE::Zen::FZenLocalServiceRunContext RunContext;
+				if (UE::Zen::TryGetLocalServiceRunContext(RunContext) && !TransferCommand->ProjectFilePath.IsEmpty())
+				{
+					UE::Zen::StartLocalService(RunContext);
+				}
+				FString DestinationProjectId, DestinationOplogId;
+				FString SplitDelim(TEXT("/"));
+				TransferCommand->Destination.Split(SplitDelim, &DestinationProjectId, &DestinationOplogId);
+				CommandLineArgs = FString::Printf(TEXT("oplog-import %s %s --clean --builds %s --namespace %s --bucket %s --builds-id %s --access-token \"%s\""),
+					*DestinationProjectId, *DestinationOplogId, *WriteToString<64>(EffectiveDomain),
+					*TransferCommand->Namespace, *TransferCommand->Bucket,
+					*WriteToString<64>(TransferCommand->BuildId), *GetAccessToken());
+			}
+			break;
+		}
+
+		FMonitoredProcess MonitoredUtilityProcess(ZenUtilityPath, *CommandLineArgs, FPaths::GetPath(ZenUtilityPath), true);
+		MonitoredUtilityProcess.OnOutput().BindSPLambda(TransferCommand, [TransferCommand](FString Output)
+		{
+			FWriteScopeLock _(TransferCommand->Lock);
+			TransferCommand->Output += TEXT("\n");
+			TransferCommand->Output += Output;
+		});
+
+		{
+			FWriteScopeLock _(TransferCommand->Lock);
+			TransferCommand->Status = EBuildTransferStatus::Active;
+		}
+		if (!MonitoredUtilityProcess.Launch())
+		{
+			UE_LOG(LogBuildServiceInstance, Warning, TEXT("Failed to launch zen utility to download build: '%s'."), *ZenUtilityPath);
+			FWriteScopeLock _(TransferCommand->Lock);
+			TransferCommand->Status = EBuildTransferStatus::Failed;
+			continue;
+		}
+
+		const uint64 StartTime = FPlatformTime::Cycles64();
+		while (MonitoredUtilityProcess.Update())
+		{
+			if (bBuildTransferThreadStopping.load(std::memory_order_relaxed))
+			{
+				return;
+			}
+			double Duration = FPlatformTime::ToSeconds64(FPlatformTime::Cycles64() - StartTime);
+			if (TransferCommand->bCancelRequested.load(std::memory_order_relaxed))
+			{
+				MonitoredUtilityProcess.Cancel(true);
+				UE_LOG(LogBuildServiceInstance, Display, TEXT("Canceled launch of zen utility for downloading build data: '%s'."), *ZenUtilityPath);
+
+				// Wait for execution to be terminated
+				while (MonitoredUtilityProcess.Update())
+				{
+					if (bBuildTransferThreadStopping.load(std::memory_order_relaxed))
+					{
+						return;
+					}
+					Duration = FPlatformTime::ToSeconds64(FPlatformTime::Cycles64() - StartTime);
+					if (Duration > 15.0)
+					{
+						UE_LOG(LogBuildServiceInstance, Warning, TEXT("Canceled launch of zen utility for downloading build data: '%s'. Failed waiting for termination."), *ZenUtilityPath);
+						break;
+					}
+					FPlatformProcess::Sleep(0.2f);
+				}
+
+				UE_LOG(LogBuildServiceInstance, Display, TEXT("Launch of zen utility for downloading build data: '%s' canceled."), *ZenUtilityPath);
+				FWriteScopeLock _(TransferCommand->Lock);
+				TransferCommand->Status = EBuildTransferStatus::Canceled;
+				continue;
+			}
+			FPlatformProcess::Sleep(0.1f);
+		}
+
+		FString OutputString = MonitoredUtilityProcess.GetFullOutputWithoutDelegate();
+		if (MonitoredUtilityProcess.GetReturnCode() == 0)
+		{
+			FWriteScopeLock _(TransferCommand->Lock);
+			TransferCommand->ReturnCode = MonitoredUtilityProcess.GetReturnCode();
+			TransferCommand->Status = EBuildTransferStatus::Succeeded;
+		}
+		else
+		{
+			UE_LOG(LogBuildServiceInstance, Warning, TEXT("Unexpected return code after launch of zen utility for downloading build data: '%s' (%d). Output: '%s'"), *ZenUtilityPath, MonitoredUtilityProcess.GetReturnCode(), *TransferCommand->Output);
+			FWriteScopeLock _(TransferCommand->Lock);
+			TransferCommand->ReturnCode = MonitoredUtilityProcess.GetReturnCode();
+			TransferCommand->Status = EBuildTransferStatus::Failed;
+		}
+	}
+}
+
 } // namespace UE::StorageService::Build
diff --git a/Engine/Source/Developer/Zen/Public/Experimental/BuildServerInterface.h b/Engine/Source/Developer/Zen/Public/Experimental/BuildServerInterface.h
--- a/Engine/Source/Developer/Zen/Public/Experimental/BuildServerInterface.h
+++ b/Engine/Source/Developer/Zen/Public/Experimental/BuildServerInterface.h
@@ -3,13 +3,17 @@
 #pragma once
 
 #include "Containers/Map.h"
+#include "Containers/MpscQueue.h"
 #include "Containers/StringFwd.h"
 #include "Containers/StringView.h"
 #include "Containers/Ticker.h"
 #include "Containers/UnrealString.h"
 #include "HAL/CriticalSection.h"
+#include "HAL/PlatformProcess.h"
+#include "HAL/Thread.h"
 #include "Serialization/CompactBinary.h"
 #include "Templates/Function.h"
+#include "Templates/SharedPointer.h"
 #include <atomic>
 
 #define UE_API ZEN_API
@@ -37,6 +41,41 @@ private:
 
 class FBuildServiceInstance
 {
+public:
+	enum class EBuildTransferStatus
+	{
+		Invalid,
+		Queued,
+		Active,
+		Failed,
+		Canceled,
+		Succeeded
+	};
+
+private:
+	enum class EBuildTransferType
+	{
+		Files,
+		Oplog
+	};
+
+	struct FBuildTransferCommand : public TSharedFromThis<FBuildTransferCommand, ESPMode::ThreadSafe>
+	{
+		EBuildTransferType Type;
+		FCbObjectId BuildId;
+		FString Namespace;
+		FString Bucket;
+		FString Destination;
+		FString ProjectFilePath;
+
+		std::atomic<bool> bCancelRequested = false;
+
+		FRWLock Lock;
+		FString Output;
+		EBuildTransferStatus Status = EBuildTransferStatus::Queued;
+		int ReturnCode;
+	};
+
 public:
 	enum class EConnectionState
 	{
@@ -52,6 +91,16 @@ public:
 		FCbObject Metadata;
 	};
 
+	struct FBuildTransfer
+	{
+	public:
+		UE_API FString GetDescription() const;
+		UE_API EBuildTransferStatus GetStatus() const;
+	private:
+		TSharedPtr<FBuildTransferCommand> BuildTransferCommand;
+		friend class UE::Zen::Build::FBuildServiceInstance;
+	};
+
 	typedef TUniqueFunction<void(EConnectionState)> FOnConnectionComplete;
 	typedef TUniqueFunction<void()> FOnRefreshNamespacesAndBucketsComplete;
 	typedef TUniqueFunction<void(TArray<FBuildRecord>&& ListBuildRecords)> FOnListBuildsComplete;
@@ -72,6 +121,9 @@ public:
 	inline const TMultiMap<FString, FString>& GetNamespacesAndBuckets() const { return NamespacesAndBuckets; }
 
 	UE_API void ListBuilds(FStringView Namespace, FStringView Bucket, FOnListBuildsComplete&& InOnListBuildsComplete);
+
+	UE_API FBuildTransfer StartBuildTransfer(const FCbObjectId& BuildId, FStringView DestinationFolder, FStringView Namespace, FStringView Bucket);
+	UE_API FBuildTransfer StartOplogBuildTransfer(const FCbObjectId& BuildId, FStringView DestinationProjectId, FStringView DestinationOplogId, FStringView ProjectFilePath, FStringView Namespace, FStringView Bucket);
 private:
 	typedef TUniqueFunction<void(FString& CommandLineArgs)> FPreZenUtilityInvocation;
 	typedef TUniqueFunction<void(bool)> FOnZenUtilityInvocationComplete;
@@ -83,6 +135,9 @@ private:
 	void InvokeZenUtility(FStringView CommandlineArgs, FPreZenUtilityInvocation&& PreInvocation, FOnZenUtilityInvocationComplete&& OnComplete);
 	void CallOnRefreshNamespacesAndBucketsComplete(FOnRefreshNamespacesAndBucketsComplete&& InOnRefreshNamespacesAndBucketsComplete);
 
+	void KickBuildTransferThread();
+	void BuildTransferThreadLoop();
+
 	TMulticastDelegate<void()> RefreshNamespacesAndBucketsComplete;
 	std::atomic<EConnectionState> ConnectionState = EConnectionState::NotStarted;
 
@@ -97,6 +152,11 @@ private:
 	FRWLock NamespacesAndBucketsLock;
 	TMultiMap<FString, FString> NamespacesAndBuckets;
 
+	TMpscQueue<TSharedRef<FBuildTransferCommand>> BuildTransferThreadCommands;
+	FThread BuildTransferThread;
+	std::atomic<bool> bBuildTransferThreadStarting = false;
+	std::atomic<bool> bBuildTransferThreadStopping = false;
+
 	FServiceSettings Settings;
 	FAnsiStringBuilderBase EffectiveDomain;
 };
diff --git a/Engine/Source/Editor/StorageServerWidgets/Internal/SBuildActivity.h b/Engine/Source/Editor/StorageServerWidgets/Internal/SBuildActivity.h
--- a/Engine/Source/Editor/StorageServerWidgets/Internal/SBuildActivity.h
+++ b/Engine/Source/Editor/StorageServerWidgets/Internal/SBuildActivity.h
@@ -9,6 +9,10 @@
 #include "Widgets/SBoxPanel.h"
 #include "Widgets/SCompoundWidget.h"
 #include "Widgets/SWidget.h"
+#include "Widgets/Views/ITableRow.h"
+#include "Widgets/Views/SListView.h"
+#include "Widgets/Views/STableRow.h"
+#include "Widgets/Views/STableViewBase.h"
 
 class STORAGESERVERWIDGETS_API SBuildActivity : public SCompoundWidget
 {
@@ -22,18 +26,52 @@ class STORAGESERVERWIDGETS_API SBuildActivity : public SCompoundWidget
 
 	SLATE_END_ARGS()
 
-	struct FBuildActivityItem : public TSharedFromThis<FBuildActivityItem, ESPMode::ThreadSafe>
+	struct FBuildActivity : public TSharedFromThis<FBuildActivity, ESPMode::ThreadSafe>
 	{
+		FString Name;
+		FString Platform;
+		UE::Zen::Build::FBuildServiceInstance::FBuildTransfer Transfer;
 	};
 
 	void Construct(const FArguments& InArgs);
 
+	void AddBuildTransfer(UE::Zen::Build::FBuildServiceInstance::FBuildTransfer Transfer, FStringView Name, FStringView Platform)
+	{
+		TSharedPtr<SBuildActivity::FBuildActivity, ESPMode::ThreadSafe> NewActivity = MakeShared<FBuildActivity>();
+		NewActivity->Name = Name;
+		NewActivity->Platform = Platform;
+		NewActivity->Transfer = Transfer;
+		BuildActivities.Add(NewActivity);
+	}
+
 private:
+	TSharedRef<ITableRow> GenerateBuildActivityRow(TSharedPtr<FBuildActivity> InItem, const TSharedRef<STableViewBase>& InOwningTable);
 	TSharedRef<SWidget> GetGridPanel();
 
+	TArray<TSharedPtr<SBuildActivity::FBuildActivity, ESPMode::ThreadSafe>> BuildActivities;
+
 	SVerticalBox::FSlot* GridSlot = nullptr;
 	TAttribute<TSharedPtr<UE::Zen::FZenServiceInstance>> ZenServiceInstance;
 	TAttribute<TSharedPtr<UE::Zen::Build::FBuildServiceInstance>> BuildServiceInstance;
 };
 
-typedef TSharedPtr<SBuildActivity::FBuildActivityItem, ESPMode::ThreadSafe> FBuildActivityItemPtr;
\ No newline at end of file
+typedef TSharedPtr<SBuildActivity::FBuildActivity, ESPMode::ThreadSafe> FBuildActivityPtr;
+
+class SBuildActivityTableRow : public SMultiColumnTableRow<FBuildActivityPtr>
+{
+public:
+	SLATE_BEGIN_ARGS(SBuildActivityTableRow) { }
+	SLATE_END_ARGS()
+
+	void Construct(const FArguments& InArgs,
+		const TSharedRef<STableViewBase>& InOwnerTableView,
+		const FBuildActivityPtr InBuildActivity,
+		TSharedPtr<UE::Zen::Build::FBuildServiceInstance> InBuildServiceInstance);
+	virtual TSharedRef<SWidget> GenerateWidgetForColumn(const FName& ColumnName) override;
+	const FSlateBrush* GetBorder() const;
+	FReply OnBrowseClicked();
+
+private:
+	FBuildActivityPtr BuildActivity;
+	TSharedPtr<UE::Zen::Build::FBuildServiceInstance> BuildServiceInstance;
+};
diff --git a/Engine/Source/Editor/StorageServerWidgets/Internal/SBuildDestination.h b/Engine/Source/Editor/StorageServerWidgets/Internal/SBuildDestination.h
deleted file mode 100644
diff --git a/Engine/Source/Editor/StorageServerWidgets/Internal/SBuildSelection.h b/Engine/Source/Editor/StorageServerWidgets/Internal/SBuildSelection.h
--- a/Engine/Source/Editor/StorageServerWidgets/Internal/SBuildSelection.h
+++ b/Engine/Source/Editor/StorageServerWidgets/Internal/SBuildSelection.h
@@ -7,16 +7,23 @@
 #include "Templates/SharedPointer.h"
 #include "Widgets/DeclarativeSyntaxSupport.h"
 #include "Widgets/Input/SComboBox.h"
+#include "Widgets/Layout/SGridPanel.h"
 #include "Widgets/SBoxPanel.h"
 #include "Widgets/SCompoundWidget.h"
 #include "Widgets/SWidget.h"
 #include "Widgets/Views/ITableRow.h"
+#include "Widgets/Views/SListView.h"
+#include "Widgets/Views/STableRow.h"
 #include "Widgets/Views/STableViewBase.h"
 
 class SMultiSelectComboBox;
 
 namespace UE::BuildSelection::Internal{ struct FListBuildsState; }
 
+DECLARE_DELEGATE_ThreeParams(FOnBuildTransferStarted, UE::Zen::Build::FBuildServiceInstance::FBuildTransfer /*Transfer*/,
+	FStringView /*Name*/,
+	FStringView /*Platform*/);
+
 class STORAGESERVERWIDGETS_API SBuildSelection : public SCompoundWidget
 {
 	SLATE_BEGIN_ARGS(SBuildSelection)
@@ -27,6 +34,7 @@ class STORAGESERVERWIDGETS_API SBuildSelection : public SCompoundWidget
 	SLATE_ATTRIBUTE(TSharedPtr<UE::Zen::FZenServiceInstance>, ZenServiceInstance);
 	SLATE_ATTRIBUTE(TSharedPtr<UE::Zen::Build::FBuildServiceInstance>, BuildServiceInstance);
 
+	SLATE_EVENT(FOnBuildTransferStarted, OnBuildTransferStarted)
 	SLATE_END_ARGS()
 
 	struct FBuildGroup : public TSharedFromThis<FBuildGroup, ESPMode::ThreadSafe>
@@ -34,6 +42,8 @@ class STORAGESERVERWIDGETS_API SBuildSelection : public SCompoundWidget
 		FString Namespace;
 		FString DisplayName;
 		FString CommitIdentifier;
+		FString Suffix;
+		FDateTime CreatedAt;
 		FString Category;
 		TMap<FString, UE::Zen::Build::FBuildServiceInstance::FBuildRecord> PerPlatformBuilds;
 	};
@@ -41,17 +51,37 @@ class STORAGESERVERWIDGETS_API SBuildSelection : public SCompoundWidget
 	void Construct(const FArguments& InArgs);
 
 private:
+	enum class EBuildType
+	{
+		Oplog,
+		StagedBuild,
+		PackagedBuild,
+		EditorPreCompiledBinary,
+		EditorInstalledBuild,
+		Unknown
+	};
+	EBuildType GetSelectedBuildType() const;
 	void RebuildLists();
+	void ReselectDestination(TSharedPtr<FBuildGroup> Item);
 	void RegenerateBuildGroups(UE::BuildSelection::Internal::FListBuildsState& ListBuildsState);
+	void RegenerateActivePlatformFilters();
+	void ValidateBuildGroupSelection();
 	TSharedRef<SWidget> OnGenerateTextBlockFromString(TSharedPtr<FString> Item);
+	bool BuildGroupIsSelectableOrNavigable(TSharedPtr<FBuildGroup> InItem) const;
 	TSharedRef<ITableRow> GenerateBuildGroupRow(TSharedPtr<FBuildGroup> InItem, const TSharedRef<STableViewBase>& InOwningTable);
+	void OnOpenDestinationDirectoryClicked();
+	TSharedRef<SWidget> GetBuildDestinationPanel();
 	TSharedRef<SWidget> GetGridPanel();
+	void BuildGroupSelectionChanged(TSharedPtr<FBuildGroup> Item, ESelectInfo::Type SelectInfo);
 
 	TSharedPtr<SComboBox<TSharedPtr<FString>>> ProjectWidget;
 	TSharedPtr<SComboBox<TSharedPtr<FString>>> StreamWidget;
 	TSharedPtr<SComboBox<TSharedPtr<FString>>> BuildTypeWidget;
 	TSharedPtr<SMultiSelectComboBox> RequiredPlatformsWidget;
 
+	TSharedPtr<SGridPanel> SelectedGroupPlatformGrid;
+	TArray<FString> SelectedGroupSelectedPlatforms;
+
 	TMultiMap<FString, FString> BucketsToNamespaces;
 
 	TArray<TSharedPtr<FString>> ProjectList;
@@ -59,6 +89,10 @@ private:
 	TArray<TSharedPtr<FString>> BuildTypeList;
 	TArray<TSharedPtr<FString>> PlatformList;
 
+	TArray<FString> ActivePlatformFilters;
+	FString DestinationFolderPath;
+	FString DestinationZenProjectId;
+
 	TSharedPtr<FString> SelectedStream;
 	TSharedPtr<FString> SelectedProject;
 	TSharedPtr<FString> SelectedBuildType;
@@ -72,6 +106,22 @@ private:
 	SVerticalBox::FSlot* GridSlot = nullptr;
 	TAttribute<TSharedPtr<UE::Zen::FZenServiceInstance>> ZenServiceInstance;
 	TAttribute<TSharedPtr<UE::Zen::Build::FBuildServiceInstance>> BuildServiceInstance;
+	FOnBuildTransferStarted OnBuildTransferStarted;
 };
 
-typedef TSharedPtr<SBuildSelection::FBuildGroup, ESPMode::ThreadSafe> FBuildSelectionBuildGroupPtr;
\ No newline at end of file
+typedef TSharedPtr<SBuildSelection::FBuildGroup, ESPMode::ThreadSafe> FBuildSelectionBuildGroupPtr;
+
+class SBuildGroupTableRow : public SMultiColumnTableRow<FBuildSelectionBuildGroupPtr>
+{
+public:
+	SLATE_BEGIN_ARGS(SBuildGroupTableRow) { }
+	SLATE_END_ARGS()
+
+	void Construct(const FArguments& InArgs, const TSharedRef<STableViewBase>& InOwnerTableView, const FBuildSelectionBuildGroupPtr InBuildGroup);
+	virtual TSharedRef<SWidget> GenerateWidgetForColumn(const FName& ColumnName) override;
+	const FSlateBrush* GetBorder() const;
+	FReply OnBrowseClicked();
+
+private:
+	FBuildSelectionBuildGroupPtr BuildGroup;
+};
diff --git a/Engine/Source/Editor/StorageServerWidgets/Private/SBuildActivity.cpp b/Engine/Source/Editor/StorageServerWidgets/Private/SBuildActivity.cpp
--- a/Engine/Source/Editor/StorageServerWidgets/Private/SBuildActivity.cpp
+++ b/Engine/Source/Editor/StorageServerWidgets/Private/SBuildActivity.cpp
@@ -17,6 +17,17 @@
 
 #define LOCTEXT_NAMESPACE "StorageServerBuild"
 
+namespace UE::BuildActivity::Internal
+{
+	namespace FBuildActivityIds
+	{
+		const FName ColStatus = TEXT("Status");
+		const FName ColName = TEXT("Name");
+		const FName ColPlatform = TEXT("Platform");
+		const FName ColDescription = TEXT("Description");
+	}
+}
+
 void SBuildActivity::Construct(const FArguments& InArgs)
 {
 	ZenServiceInstance = InArgs._ZenServiceInstance;
@@ -36,16 +47,28 @@ void SBuildActivity::Construct(const FArguments& InArgs)
 	];
 }
 
+TSharedRef<ITableRow> SBuildActivity::GenerateBuildActivityRow(FBuildActivityPtr InItem, const TSharedRef<STableViewBase>& InOwningTable)
+{
+	return SNew(SBuildActivityTableRow, InOwningTable, InItem, BuildServiceInstance.Get());
+}
+
 TSharedRef<SWidget> SBuildActivity::GetGridPanel()
 {
-	TSharedRef<SGridPanel> Panel = SNew(SGridPanel);
+	using namespace UE::BuildActivity::Internal;
+
+	TSharedRef<SVerticalBox> Panel = SNew(SVerticalBox)
+	.Visibility_Lambda([this]()
+	{
+		return BuildActivities.IsEmpty() ? EVisibility::Collapsed : EVisibility::Visible;
+	});
 
 	const float RowMargin = 0.0f;
 	const float ColumnMargin = 10.0f;
 	const FSlateColor TitleColor = FStyleColors::AccentWhite;
 	const FSlateFontInfo TitleFont = FCoreStyle::GetDefaultFontStyle("Bold", 10);
 
-	Panel->AddSlot(0, 0)
+	Panel->AddSlot()
+	.AutoHeight()
 	[
 		SNew(STextBlock)
 		.Margin(FMargin(ColumnMargin, RowMargin))
@@ -54,14 +77,126 @@ TSharedRef<SWidget> SBuildActivity::GetGridPanel()
 		.Text(LOCTEXT("BuildActivity_Activity", "Activity"))
 	];
 
-	Panel->AddSlot(0, 1)
+	Panel->AddSlot()
+	.AutoHeight()
 	.Padding(FMargin(ColumnMargin, RowMargin))
 	[
-		SNew(SListView<FBuildActivityItemPtr>)
+		SNew(SListView<FBuildActivityPtr>)
+			.ListItemsSource(&BuildActivities)
+			.OnGenerateRow(this, &SBuildActivity::GenerateBuildActivityRow)
+			.HeaderRow
+			(
+				SNew(SHeaderRow)
+				.Visibility(EVisibility::Collapsed)
+				+ SHeaderRow::Column(FBuildActivityIds::ColStatus).DefaultLabel(LOCTEXT("BuildActivity_ActivityStatus", "Status"))
+					.FillWidth(0.1f)
+				+ SHeaderRow::Column(FBuildActivityIds::ColName).DefaultLabel(LOCTEXT("BuildActivity_ActivityName", "Name"))
+					.FillWidth(0.2f).HAlignCell(HAlign_Left).HAlignHeader(HAlign_Left).VAlignCell(VAlign_Center)
+				+ SHeaderRow::Column(FBuildActivityIds::ColPlatform).DefaultLabel(LOCTEXT("BuildActivity_ActivityPlatform", "Platform"))
+					.FillWidth(0.1f).HAlignCell(HAlign_Left).HAlignHeader(HAlign_Left).VAlignCell(VAlign_Center)
+				+ SHeaderRow::Column(FBuildActivityIds::ColDescription).DefaultLabel(LOCTEXT("BuildActivity_ActivityDescription", "Description"))
+					.FillWidth(0.6f).HAlignCell(HAlign_Left).HAlignHeader(HAlign_Center).VAlignCell(VAlign_Center)
+			)
 	];
 
 
 	return Panel;
 }
 
+void
+SBuildActivityTableRow::Construct(const FArguments& InArgs,
+	const TSharedRef<STableViewBase>& InOwnerTableView,
+	const FBuildActivityPtr InBuildActivity,
+	TSharedPtr<UE::Zen::Build::FBuildServiceInstance> InBuildServiceInstance)
+{
+	BuildActivity = InBuildActivity;
+	BuildServiceInstance = InBuildServiceInstance;
+
+	SMultiColumnTableRow<FBuildActivityPtr>::Construct(FSuperRowType::FArguments(), InOwnerTableView);
+}
+
+TSharedRef<SWidget>
+SBuildActivityTableRow::GenerateWidgetForColumn(const FName& ColumnName)
+{
+	using namespace UE::Zen::Build;
+	using namespace UE::BuildActivity::Internal;
+
+	if (ColumnName == FBuildActivityIds::ColStatus)
+	{
+		return SNew(SHorizontalBox)
+				+ SHorizontalBox::Slot()
+				[
+					SNew(STextBlock).Text_Lambda([this]()
+						{
+							FString Status;
+							if (BuildServiceInstance)
+							{
+								switch(BuildActivity->Transfer.GetStatus())
+								{
+								case FBuildServiceInstance::EBuildTransferStatus::Invalid:
+									return LOCTEXT("BuildActivity_TransferStatusInvalid", "Invalid");
+								case FBuildServiceInstance::EBuildTransferStatus::Queued:
+									return LOCTEXT("BuildActivity_TransferStatusQueued", "Queued");
+								case FBuildServiceInstance::EBuildTransferStatus::Active:
+									return LOCTEXT("BuildActivity_TransferStatusActive", "Active");
+								case FBuildServiceInstance::EBuildTransferStatus::Failed:
+									return LOCTEXT("BuildActivity_TransferStatusFailed", "Failed");
+								case FBuildServiceInstance::EBuildTransferStatus::Canceled:
+									return LOCTEXT("BuildActivity_TransferStatusCanceled", "Canceled");
+								case FBuildServiceInstance::EBuildTransferStatus::Succeeded:
+									return LOCTEXT("BuildActivity_TransferStatusSucceeded", "Succeeded");
+								}
+							}
+							return LOCTEXT("BuildActivity_TransferStatusUnknown", "Unknown");
+						})
+				];
+	}
+	else if (ColumnName == FBuildActivityIds::ColName)
+	{
+		return SNew(SHorizontalBox)
+				+ SHorizontalBox::Slot()
+				[
+					SNew(STextBlock).Text(FText::FromString(BuildActivity->Name))
+				];
+	}
+	else if (ColumnName == FBuildActivityIds::ColPlatform)
+	{
+		return SNew(SHorizontalBox)
+				+ SHorizontalBox::Slot()
+				[
+					SNew(STextBlock).Text(FText::FromString(BuildActivity->Platform))
+				];
+	}
+	else if (ColumnName == FBuildActivityIds::ColDescription)
+	{
+		return SNew(SHorizontalBox)
+				+ SHorizontalBox::Slot()
+				[
+					SNew(STextBlock).Text_Lambda([this]()
+						{
+							FString Description;
+							if (BuildServiceInstance)
+							{
+								Description = BuildActivity->Transfer.GetDescription();
+							}
+							return FText::FromString(Description);
+						})
+				];
+	}
+
+	return SNullWidget::NullWidget;
+}
+
+const FSlateBrush*
+SBuildActivityTableRow::GetBorder() const
+{
+	return STableRow<FBuildActivityPtr>::GetBorder();
+}
+
+FReply
+SBuildActivityTableRow::OnBrowseClicked()
+{
+	return FReply::Unhandled();
+}
+
 #undef LOCTEXT_NAMESPACE
diff --git a/Engine/Source/Editor/StorageServerWidgets/Private/SBuildDestination.cpp b/Engine/Source/Editor/StorageServerWidgets/Private/SBuildDestination.cpp
deleted file mode 100644
diff --git a/Engine/Source/Editor/StorageServerWidgets/Private/SBuildLogin.cpp b/Engine/Source/Editor/StorageServerWidgets/Private/SBuildLogin.cpp
--- a/Engine/Source/Editor/StorageServerWidgets/Private/SBuildLogin.cpp
+++ b/Engine/Source/Editor/StorageServerWidgets/Private/SBuildLogin.cpp
@@ -55,6 +55,7 @@ TSharedRef<SWidget> SBuildLogin::GetGridPanel()
 	];
 
 	Panel->AddSlot(1, 0)
+	.Padding(FMargin(ColumnMargin, RowMargin))
 	[
 		SNew(STextBlock)
 		.Text_Lambda([this]
@@ -92,10 +93,10 @@ TSharedRef<SWidget> SBuildLogin::GetGridPanel()
 	];
 
 	Panel->AddSlot(2, 0)
+	.Padding(FMargin(ColumnMargin, RowMargin))
 	[
 		SNew(SHyperlink)
-		.Padding(FMargin(ColumnMargin, RowMargin))
-		.UnderlineStyle(FAppStyle::Get(), "NoBorder")
+		.Style(FAppStyle::Get(), TEXT("NavigationHyperlink"))
 		.Text_Lambda([this]
 		{
 			FText BrowseText = LOCTEXT("BuildLogin_ConnectLink", "Click to connect");
diff --git a/Engine/Source/Editor/StorageServerWidgets/Private/SBuildSelection.cpp b/Engine/Source/Editor/StorageServerWidgets/Private/SBuildSelection.cpp
--- a/Engine/Source/Editor/StorageServerWidgets/Private/SBuildSelection.cpp
+++ b/Engine/Source/Editor/StorageServerWidgets/Private/SBuildSelection.cpp
@@ -2,14 +2,20 @@
 
 #include "SBuildSelection.h"
 
+#include "DesktopPlatformModule.h"
 #include "Internationalization/FastDecimalFormat.h"
 #include "Math/BasicMathExpressionEvaluator.h"
 #include "Math/UnitConversion.h"
+#include "Misc/App.h"
 #include "Misc/ExpressionParser.h"
+#include "Misc/Paths.h"
+#include "Misc/UProjectInfo.h"
 #include "SMultiSelectComboBox.h"
 #include "Styling/StyleColors.h"
 #include "Widgets/Images/SImage.h"
 #include "Widgets/Input/SButton.h"
+#include "Widgets/Input/SCheckBox.h"
+#include "Widgets/Input/SEditableTextBox.h"
 #include "Widgets/Input/SHyperlink.h"
 #include "Widgets/Layout/SGridPanel.h"
 #include "Widgets/SBoxPanel.h"
@@ -28,6 +34,7 @@ namespace UE::BuildSelection::Internal
 		const FName ColName = TEXT("Name");
 		const FName ColCommit = TEXT("Commit");
 		const FName ColSuffix = TEXT("Suffix");
+		const FName ColCategory = TEXT("Category");
 		const FName ColCreated = TEXT("Created");
 	}
 
@@ -56,6 +63,7 @@ void SBuildSelection::Construct(const FArguments& InArgs)
 {
 	ZenServiceInstance = InArgs._ZenServiceInstance;
 	BuildServiceInstance = InArgs._BuildServiceInstance;
+	OnBuildTransferStarted = InArgs._OnBuildTransferStarted;
 
 	if (TSharedPtr<UE::Zen::Build::FBuildServiceInstance> ServiceInstance = BuildServiceInstance.Get())
 	{
@@ -77,6 +85,35 @@ void SBuildSelection::Construct(const FArguments& InArgs)
 	];
 }
 
+SBuildSelection::EBuildType SBuildSelection::GetSelectedBuildType() const
+{
+	if (!SelectedBuildType)
+	{
+		return EBuildType::Unknown;
+	}
+	else if (SelectedBuildType->Contains(TEXT("oplog")))
+	{
+		return EBuildType::Oplog;
+	}
+	else if (SelectedBuildType->Contains(TEXT("stagedbuild")) || SelectedBuildType->Contains(TEXT("staged-build")))
+	{
+		return EBuildType::StagedBuild;
+	}
+	else if (SelectedBuildType->Contains(TEXT("packagedbuild")) || SelectedBuildType->Contains(TEXT("packaged-build")))
+	{
+		return EBuildType::PackagedBuild;
+	}
+	else if (SelectedBuildType->Contains(TEXT("ugs-pcb")))
+	{
+		return EBuildType::EditorPreCompiledBinary;
+	}
+	else if (SelectedBuildType->Contains(TEXT("installedbuild")) || SelectedBuildType->Contains(TEXT("installed-build")))
+	{
+		return EBuildType::EditorInstalledBuild;
+	}
+	return EBuildType::Unknown;
+}
+
 void SBuildSelection::RebuildLists()
 {
 	using namespace UE::Zen::Build;
@@ -96,75 +133,140 @@ void SBuildSelection::RebuildLists()
 		TArray<FString> Streams;
 		TArray<FString> BuildTypes;
 		TArray<FString> Platforms;
-		const TMultiMap<FString, FString>& NamespacesAndBuckets = ServiceInstance->GetNamespacesAndBuckets();
+		TMultiMap<FString, FString> NamespacesAndBuckets = ServiceInstance->GetNamespacesAndBuckets();
+
+		auto StringToSegmentViews = [](const FString& Str, TArray<FStringView>& OutViews)
+		{
+			FStringView WorkingStringView(Str);
+			int32 CurrentIndex = 0;
+			while (WorkingStringView.FindChar(TCHAR('.'), CurrentIndex))
+			{
+				if (CurrentIndex != 0)
+				{
+					OutViews.Add(WorkingStringView.Left(CurrentIndex));
+				}
+				WorkingStringView.RightChopInline(CurrentIndex+1);
+			}
+			if (!WorkingStringView.IsEmpty())
+			{
+				OutViews.Add(WorkingStringView);
+			}
+		};
 
 		NamespacesAndBuckets.GetKeys(Namespaces);
 
+		auto ConvertToSharedPtrs = [] (TArray<FString>& Strings, TArray<TSharedPtr<FString>>& SharedStrings)
+		{
+			Strings.StableSort();
+			for (FString& String : Strings)
+			{
+				SharedStrings.Add(MakeShared<FString>(MoveTemp(String)));
+			}
+		};
+
+		auto ConformSelection = [](TSharedPtr<FString>& SelectedItem, const TArray<TSharedPtr<FString>>& SelectionList)
+		{
+			if (!SelectedItem)
+			{
+				SelectedItem = SelectionList.IsEmpty() ? nullptr : SelectionList[0];
+				return;
+			}
+
+			const TSharedPtr<FString>* FoundSelectionListItem = SelectionList.FindByPredicate([&SelectedItem](const TSharedPtr<FString>& Item)
+			{
+				return *Item == *SelectedItem;
+			});
+
+			SelectedItem = FoundSelectionListItem ? *FoundSelectionListItem : SelectionList.IsEmpty() ? nullptr : SelectionList[0];
+		};
+
+		const uint32 SegmentIndexProject = 0;
+		const uint32 SegmentIndexBuildType = 1;
+		const uint32 SegmentIndexStream = 2;
+		const uint32 SegmentIndexPlatform = 3;
+		const uint32 SegmentIndexNum = 4;
+
+		// Stream list generation and selection conforming
+		TMultiMap<FStringView, TArray<FStringView>> NamespacesToBucketSegmentViews;
 		for (const TPair<FString, FString>& NamespaceAndBucket : NamespacesAndBuckets)
 		{
 			BucketsToNamespaces.AddUnique(NamespaceAndBucket.Value, NamespaceAndBucket.Key);
+			TArray<FStringView>& BucketSegmentViews = NamespacesToBucketSegmentViews.Add(NamespaceAndBucket.Key);
+			StringToSegmentViews(NamespaceAndBucket.Value, BucketSegmentViews);
 
-			TArray<FString> Segments;
-			NamespaceAndBucket.Value.ParseIntoArray(Segments, TEXT("."), true);
+			if (BucketSegmentViews.Num() == SegmentIndexNum)
+			{
+				Streams.AddUnique(FString(BucketSegmentViews[SegmentIndexStream]));
+			}
+		}
+		ConvertToSharedPtrs(Streams, StreamList);
+		ConformSelection(SelectedStream, StreamList);
 
-			if (Segments.Num() == 4)
+		// Project list generation and selection conforming
+		for (const TPair<FStringView, TArray<FStringView>>& NamespaceToBucketSegmentViews : NamespacesToBucketSegmentViews)
+		{
+			const TArray<FStringView>& BucketSegmentViews = NamespaceToBucketSegmentViews.Value;
+			if (BucketSegmentViews.Num() == SegmentIndexNum)
 			{
-				Projects.AddUnique(Segments[0]);
-				if (Segments[0] != (SelectedProject ? *SelectedProject : Projects[0]))
+				if (BucketSegmentViews[SegmentIndexStream] != (SelectedStream ? *SelectedStream : Streams[0]))
 				{
 					continue;
 				}
 
-				Streams.AddUnique(Segments[2]);
-				if (Segments[2] != (SelectedStream ? *SelectedStream : Streams[0]))
+				Projects.AddUnique(FString(BucketSegmentViews[SegmentIndexProject]));
+			}
+		}
+		ConvertToSharedPtrs(Projects, ProjectList);
+		ConformSelection(SelectedProject, ProjectList);
+
+		// BuildType list generation and selection conforming
+		for (const TPair<FStringView, TArray<FStringView>>& NamespaceToBucketSegmentViews : NamespacesToBucketSegmentViews)
+		{
+			const TArray<FStringView>& BucketSegmentViews = NamespaceToBucketSegmentViews.Value;
+			if (BucketSegmentViews.Num() == SegmentIndexNum)
+			{
+				if (BucketSegmentViews[SegmentIndexStream] != (SelectedStream ? *SelectedStream : Streams[0]))
 				{
 					continue;
 				}
 
-				BuildTypes.AddUnique(Segments[1]);
-				if (Segments[1] != (SelectedBuildType ? *SelectedBuildType : BuildTypes[0]))
+				if (BucketSegmentViews[SegmentIndexProject] != (SelectedProject ? *SelectedProject : Projects[0]))
 				{
 					continue;
 				}
 
-				Platforms.AddUnique(Segments[3]);
+				BuildTypes.AddUnique(FString(BucketSegmentViews[SegmentIndexBuildType]));
 			}
 		}
-
-		auto ConvertToSharedPtrs = [] (TArray<FString>& Strings, TArray<TSharedPtr<FString>>& SharedStrings)
-		{
-			Strings.StableSort();
-			for (FString& String : Strings)
-			{
-				SharedStrings.Add(MakeShared<FString>(MoveTemp(String)));
-			}
-		};
-
-		ConvertToSharedPtrs(Projects, ProjectList);
-		ConvertToSharedPtrs(Streams, StreamList);
 		ConvertToSharedPtrs(BuildTypes, BuildTypeList);
-		ConvertToSharedPtrs(Platforms, PlatformList);
-	}
+		ConformSelection(SelectedBuildType, BuildTypeList);
 
-	auto ConformSelection = [](TSharedPtr<FString>& SelectedItem, const TArray<TSharedPtr<FString>>& SelectionList)
-	{
-		if (!SelectedItem)
+		// Platform list generation
+		for (const TPair<FStringView, TArray<FStringView>>& NamespaceToBucketSegmentViews : NamespacesToBucketSegmentViews)
 		{
-			SelectedItem = SelectionList.IsEmpty() ? nullptr : SelectionList[0];
-			return;
-		}
+			const TArray<FStringView>& BucketSegmentViews = NamespaceToBucketSegmentViews.Value;
+			if (BucketSegmentViews.Num() == SegmentIndexNum)
+			{
+				if (BucketSegmentViews[SegmentIndexStream] != (SelectedStream ? *SelectedStream : Streams[0]))
+				{
+					continue;
+				}
 
-		const TSharedPtr<FString>* FoundSelectionListItem = SelectionList.FindByPredicate([&SelectedItem](const TSharedPtr<FString>& Item)
-		{
-			return *Item == *SelectedItem;
-		});
+				if (BucketSegmentViews[SegmentIndexProject] != (SelectedProject ? *SelectedProject : Projects[0]))
+				{
+					continue;
+				}
 
-		SelectedItem = FoundSelectionListItem ? *FoundSelectionListItem : SelectionList.IsEmpty() ? nullptr : SelectionList[0];
-	};
+				if (BucketSegmentViews[SegmentIndexBuildType] != (SelectedBuildType ? *SelectedBuildType : BuildTypes[0]))
+				{
+					continue;
+				}
 
-	ConformSelection(SelectedProject, ProjectList);
-	ConformSelection(SelectedStream, StreamList);
-	ConformSelection(SelectedBuildType, BuildTypeList);
+				Platforms.AddUnique(FString(BucketSegmentViews[SegmentIndexPlatform]));
+			}
+		}
+		ConvertToSharedPtrs(Platforms, PlatformList);
+	}
 
 	ExecuteOnGameThread(UE_SOURCE_LOCATION,
 		[this]
@@ -175,6 +277,7 @@ void SBuildSelection::RebuildLists()
 			ProjectWidget->SetSelectedItem(SelectedProject);
 			BuildTypeWidget->RefreshOptions();
 			BuildTypeWidget->SetSelectedItem(SelectedBuildType);
+			RegenerateActivePlatformFilters();
 		});
 
 	if (TSharedPtr<FBuildServiceInstance> ServiceInstance = BuildServiceInstance.Get())
@@ -187,7 +290,7 @@ void SBuildSelection::RebuildLists()
 			BucketsToNamespaces.MultiFind(Bucket, NamespacesForBucket);
 			for (FString& Namespace : NamespacesForBucket)
 			{
-				NamespacePlatformBucketTuples.Emplace(MoveTemp(Namespace), *Platform, MoveTemp(Bucket));
+				NamespacePlatformBucketTuples.Emplace(MoveTemp(Namespace), *Platform, Bucket);
 			}
 		}
 
@@ -211,21 +314,39 @@ void SBuildSelection::RebuildLists()
 					NewBuildState.Platform = MoveTemp(Platform);
 					NewBuildState.Results = MoveTemp(Results);
 
-					if ((ExpectedBuildRefreshGeneration == BuildRefreshGeneration) && (--PendingQueryState->PendingQueries == 0))
+					if (--PendingQueryState->PendingQueries == 0)
 					{
 						// All queries complete
-						RegenerateBuildGroups(*PendingQueryState);
-
-						ExecuteOnGameThread(UE_SOURCE_LOCATION,
-							[this]
-							{
-								BuildListView->RequestListRefresh();
-								BuildListRefreshesInProgress.fetch_sub(1);
-							});
+						if (ExpectedBuildRefreshGeneration == BuildRefreshGeneration)
+						{
+							// Expected generation is the current generation
+							RegenerateBuildGroups(*PendingQueryState);
+
+							ExecuteOnGameThread(UE_SOURCE_LOCATION,
+								[this]
+								{
+									BuildListView->RequestListRefresh();
+									BuildListRefreshesInProgress.fetch_sub(1);
+								});
+						}
+						else
+						{
+							// Expected generation is not the current generation
+							ExecuteOnGameThread(UE_SOURCE_LOCATION,
+								[this]
+								{
+									BuildListRefreshesInProgress.fetch_sub(1);
+								});
+						}
 					}
 				});
 			++QueryIndex;
 		}
+
+		if (NamespacePlatformBucketTuples.IsEmpty())
+		{
+			BuildListRefreshesInProgress.fetch_sub(1);
+		}
 	}
 	else
 	{
@@ -233,6 +354,42 @@ void SBuildSelection::RebuildLists()
 	}
 }
 
+void SBuildSelection::ReselectDestination(TSharedPtr<FBuildGroup> Item)
+{
+	// TODO: Select the destination info (path & zen project id based on the selected item)
+	if (!SelectedBuildType)
+	{
+		return;
+	}
+
+	EBuildType BuildType = GetSelectedBuildType();
+	FString ProjectFilename = FUProjectDictionary::GetDefault().GetProjectPathForGame(**SelectedProject);
+	if (ProjectFilename.IsEmpty())
+	{
+		if (BuildType != EBuildType::Oplog)
+		{
+			DestinationFolderPath = FPaths::Combine(FPaths::ProjectSavedDir(), TEXT("DownloadedBuilds"));
+		}
+		return;
+	}
+
+	switch (BuildType)
+	{
+	case EBuildType::Oplog:
+		DestinationZenProjectId = FApp::GetZenStoreProjectIdForProject(ProjectFilename);
+		break;
+	case EBuildType::StagedBuild:
+		DestinationFolderPath = FPaths::Combine(FPaths::GetPath(ProjectFilename), TEXT("Saved"), TEXT("StagedBuilds"));
+		break;
+	case EBuildType::PackagedBuild:
+		DestinationFolderPath = FPaths::Combine(FPaths::GetPath(ProjectFilename), TEXT("Saved"), TEXT("Packages"));
+		break;
+	default:
+		DestinationFolderPath = FPaths::Combine(FPaths::GetPath(ProjectFilename), TEXT("Saved"), TEXT("DownloadedBuilds"));
+		break;
+	}
+}
+
 void SBuildSelection::RegenerateBuildGroups(UE::BuildSelection::Internal::FListBuildsState& ListBuildsState)
 {
 	using namespace UE::Zen::Build;
@@ -246,7 +403,7 @@ void SBuildSelection::RegenerateBuildGroups(UE::BuildSelection::Internal::FListB
 		{
 			return FString(WriteToString<64>(Namespace, ".", CommitIdentifier));
 		}
-		if (FCbFieldView NameField = BuildRecord.Metadata["name"]; !NameField.HasError())
+		if (FCbFieldView NameField = BuildRecord.Metadata["name"]; NameField.HasValue() && !NameField.HasError())
 		{
 			return FString(WriteToString<64>(Namespace, ".", NameField.AsString()));
 		}
@@ -259,7 +416,7 @@ void SBuildSelection::RegenerateBuildGroups(UE::BuildSelection::Internal::FListB
 		for (FBuildServiceInstance::FBuildRecord& BuildRecord : BuildState.Results)
 		{
 			FString CommitIdentifier;
-			if (FCbFieldView ChangelistField = BuildRecord.Metadata["changelist"]; !ChangelistField.HasError())
+			if (FCbFieldView ChangelistField = BuildRecord.Metadata["changelist"]; ChangelistField.HasValue() && !ChangelistField.HasError())
 			{
 				if (ChangelistField.IsString())
 				{
@@ -269,8 +426,12 @@ void SBuildSelection::RegenerateBuildGroups(UE::BuildSelection::Internal::FListB
 				{
 					CommitIdentifier = *WriteToString<64>(ChangelistField.AsUInt64());
 				}
+				else if (ChangelistField.IsFloat())
+				{
+					CommitIdentifier = *WriteToString<64>((uint64)ChangelistField.AsDouble());
+				}
 			}
-			else if (FCbFieldView CommitField = BuildRecord.Metadata["commit"]; !CommitField.HasError())
+			else if (FCbFieldView CommitField = BuildRecord.Metadata["commit"]; CommitField.HasValue() && !CommitField.HasError())
 			{
 				if (CommitField.IsString())
 				{
@@ -280,6 +441,10 @@ void SBuildSelection::RegenerateBuildGroups(UE::BuildSelection::Internal::FListB
 				{
 					CommitIdentifier = *WriteToString<64>(CommitField.AsUInt64());
 				}
+				else if (CommitField.IsFloat())
+				{
+					CommitIdentifier = *WriteToString<64>((uint64)CommitField.AsDouble());
+				}
 			}
 			FString GroupKey = MakeGroupKey(BuildState.Namespace, CommitIdentifier, BuildRecord);
 			TSharedPtr<FBuildGroup>& BuildGroup = KeyedGroups.FindOrAdd(GroupKey);
@@ -291,13 +456,26 @@ void SBuildSelection::RegenerateBuildGroups(UE::BuildSelection::Internal::FListB
 			if (BuildGroup->DisplayName.IsEmpty())
 			{
 				FString Category;
-				if (FCbFieldView TemplateIdField = BuildRecord.Metadata["hordeTemplateId"]; !TemplateIdField.HasError())
+				if (FCbFieldView TemplateIdField = BuildRecord.Metadata["hordeTemplateId"]; TemplateIdField.HasValue() && !TemplateIdField.HasError())
 				{
 					Category = FUTF8ToTCHAR(TemplateIdField.AsString());
 				}
 
+				FDateTime CreatedAt;
+				if (FCbFieldView CreatedAtField = BuildRecord.Metadata["createdAt"]; CreatedAtField.HasValue() && !CreatedAtField.HasError())
+				{
+					if (CreatedAtField.IsString())
+					{
+						FDateTime::ParseIso8601(FUTF8ToTCHAR(CreatedAtField.AsString()).Get(), CreatedAt);
+					}
+					else if (CreatedAtField.IsDateTime())
+					{
+						CreatedAt = CreatedAtField.AsDateTime();
+					}
+				}
+
 				FString ItemName;
-				if (FCbFieldView NameView = BuildRecord.Metadata["name"]; !NameView.HasError())
+				if (FCbFieldView NameView = BuildRecord.Metadata["name"]; NameView.HasValue() && !NameView.HasError())
 				{
 					// TODO: This name manipulation needs to be removed when the metadata is more consistent.
 					ItemName = FUTF8ToTCHAR(NameView.AsString());
@@ -339,9 +517,11 @@ void SBuildSelection::RegenerateBuildGroups(UE::BuildSelection::Internal::FListB
 					ItemName = *WriteToString<64>(BuildRecord.BuildId);
 				}
 
+				BuildGroup->Namespace = BuildState.Namespace;
 				BuildGroup->DisplayName = ItemName;
 				BuildGroup->CommitIdentifier = CommitIdentifier;
-				BuildGroup->Namespace = BuildState.Namespace;
+				BuildGroup->Category = Category;
+				BuildGroup->CreatedAt = CreatedAt;
 			}
 
 			BuildGroup->PerPlatformBuilds.FindOrAdd(BuildState.Platform, MoveTemp(BuildRecord));
@@ -351,26 +531,279 @@ void SBuildSelection::RegenerateBuildGroups(UE::BuildSelection::Internal::FListB
 	KeyedGroups.GenerateValueArray(BuildGroups);
 }
 
+void SBuildSelection::RegenerateActivePlatformFilters()
+{
+	ActivePlatformFilters.Empty();
+	for (TSharedPtr<FString> Platform : PlatformList)
+	{
+		if (Platform && RequiredPlatformsWidget)
+		{
+			if (RequiredPlatformsWidget->IsChecked(*Platform))
+			{
+				ActivePlatformFilters.Add(*Platform);
+				SelectedGroupSelectedPlatforms.AddUnique(*Platform);
+			}
+			else
+			{
+				SelectedGroupSelectedPlatforms.Remove(*Platform);
+			}
+		}
+	}
+}
+
+void SBuildSelection::ValidateBuildGroupSelection()
+{
+	BuildListView->UpdateSelectionSet();
+	TArray<FBuildSelectionBuildGroupPtr> SelectedItems = BuildListView->GetSelectedItems();
+	if (SelectedItems.IsEmpty())
+	{
+		return;
+	}
+	for (const FBuildSelectionBuildGroupPtr& SelectedItem : SelectedItems)
+	{
+		if (!BuildGroupIsSelectableOrNavigable(SelectedItem))
+		{
+			BuildListView->SetItemSelection(SelectedItem, false);
+		}
+	}
+}
+
 TSharedRef<SWidget> SBuildSelection::OnGenerateTextBlockFromString(TSharedPtr<FString> Item)
 {
 	return SNew(STextBlock)
 		.Text(FText::FromString(*Item));
 }
 
+bool SBuildSelection::BuildGroupIsSelectableOrNavigable(FBuildSelectionBuildGroupPtr InItem) const
+{
+	if (!InItem)
+	{
+		return false;
+	}
+
+	for (const FString& ActivePlatformFilter : ActivePlatformFilters)
+	{
+		if (!InItem->PerPlatformBuilds.Contains(ActivePlatformFilter))
+		{
+			return false;
+		}
+	}
+	return true;
+}
+
 TSharedRef<ITableRow> SBuildSelection::GenerateBuildGroupRow(FBuildSelectionBuildGroupPtr InItem, const TSharedRef<STableViewBase>& InOwningTable)
 {
-	return SNew(STableRow<FBuildSelectionBuildGroupPtr>, InOwningTable)
-		[
-			SNew(SHorizontalBox)
-				+ SHorizontalBox::Slot()
+	return SNew(SBuildGroupTableRow, InOwningTable, InItem)
+		.Visibility_Lambda([this, InItem]()
+		{
+			for (const FString& ActivePlatformFilter : ActivePlatformFilters)
+			{
+				if (!InItem->PerPlatformBuilds.Contains(ActivePlatformFilter))
+				{
+					return EVisibility::Collapsed;
+				}
+			}
+			return EVisibility::Visible;
+		});
+}
+
+void SBuildSelection::BuildGroupSelectionChanged(FBuildSelectionBuildGroupPtr Item, ESelectInfo::Type SelectInfo)
+{
+	using namespace UE::Zen::Build;
+
+	if (!SelectedGroupPlatformGrid)
+	{
+		return;
+	}
+
+	SelectedGroupPlatformGrid->ClearChildren();
+
+	if (!Item)
+	{
+		return;
+	}
+
+	ReselectDestination(Item);
+
+	int32 Row = 0;
+	int32 Column = 0;
+
+	for (const TSharedPtr<FString>& Platform : PlatformList)
+	{
+		TSharedPtr<SCheckBox> AssociatedCheckbox;
+		SelectedGroupPlatformGrid->AddSlot(Column, Row)
+			.Padding(0,0,0,2)
+			[
+				SNew(SHorizontalBox)
+				.IsEnabled_Lambda([this, Item, Platform]
+				{
+					return Item->PerPlatformBuilds.Contains(*Platform);
+				})
+				+SHorizontalBox::Slot()
+				.AutoWidth()
+				.VAlign(VAlign_Center)
 				[
-					SNew(STextBlock).Text(FText::FromString(InItem->DisplayName))
+					SAssignNew(AssociatedCheckbox, SCheckBox)
+					.IsChecked_Lambda([this, Item, Platform]
+					{
+						bool bIsChecked = Item->PerPlatformBuilds.Contains(*Platform) && SelectedGroupSelectedPlatforms.Contains(*Platform);
+						return bIsChecked ? ECheckBoxState::Checked : ECheckBoxState::Unchecked;
+					})
+					.OnCheckStateChanged_Lambda([this, Platform](ECheckBoxState InNewState)
+					{
+						if (InNewState == ECheckBoxState::Checked)
+						{
+							SelectedGroupSelectedPlatforms.AddUnique(*Platform);
+						}
+						else if (InNewState == ECheckBoxState::Unchecked)
+						{
+							SelectedGroupSelectedPlatforms.Remove(*Platform);
+						}
+					})
 				]
-				+ SHorizontalBox::Slot()
+				+SHorizontalBox::Slot()
+				.AutoWidth()
+				.VAlign(VAlign_Center)
 				[
-					SNew(STextBlock).Text(FText::FromString(InItem->CommitIdentifier))
+					SNew(SButton)
+					.ButtonStyle(FAppStyle::Get(), "InvisibleButton")
+					.IsFocusable(false)
+
+					.OnClicked_Lambda([this, AssociatedCheckbox, Platform]()
+						{
+							ECheckBoxState NewState = ECheckBoxState::Checked;
+							if (AssociatedCheckbox->IsChecked())
+							{
+								NewState = ECheckBoxState::Unchecked;
+							}
+							AssociatedCheckbox->SetIsChecked(NewState);
+							if (NewState == ECheckBoxState::Checked)
+							{
+								SelectedGroupSelectedPlatforms.AddUnique(*Platform);
+							}
+							else if (NewState == ECheckBoxState::Unchecked)
+							{
+								SelectedGroupSelectedPlatforms.Remove(*Platform);
+							}
+							return FReply::Handled();
+						})
+					[
+						SNew(STextBlock)
+						.Justification(ETextJustify::Left)
+						.Text(FText::FromString(*Platform))
+					]
 				]
-		];
+			];
+
+		if (++Row > 2)
+		{
+			Column++;
+			Row = 0;
+		}
+	}
+}
+
+void SBuildSelection::OnOpenDestinationDirectoryClicked()
+{
+	bool bDirectorySelected = false;
+
+	if (IDesktopPlatform* DesktopPlatform = FDesktopPlatformModule::Get())
+	{
+		const FString Title = LOCTEXT("BuildSelection_DestinationDirectoryBrowserTitle", "Choose destination directory").ToString();
+
+		bDirectorySelected = DesktopPlatform->OpenDirectoryDialog(
+			FSlateApplication::Get().FindBestParentWindowHandleForDialogs(nullptr),
+			Title,
+			DestinationFolderPath,
+			DestinationFolderPath);
+	}
+}
+
+TSharedRef<SWidget> SBuildSelection::GetBuildDestinationPanel()
+{
+	return
+	SNew(SVerticalBox)
+
+	+ SVerticalBox::Slot()
+	.AutoHeight()
+	.VAlign(VAlign_Top)
+	[
+		SNew(SHorizontalBox)
+		.Visibility_Lambda([this]()
+		{
+			TArray<FBuildSelectionBuildGroupPtr> SelectedItems = BuildListView->GetSelectedItems();
+			return (SelectedItems.Num() == 0) || GetSelectedBuildType() == EBuildType::Oplog ? EVisibility::Collapsed : EVisibility::Visible;
+		})
+		+ SHorizontalBox::Slot()
+		.FillWidth(1.0f)
+		.VAlign(VAlign_Center)
+		[
+			SNew(SEditableTextBox)
+			.OverflowPolicy(ETextOverflowPolicy::MiddleEllipsis)
+			.MinDesiredWidth(200.0f)
+			.Text_Lambda([this]()
+			{
+				return FText::FromString(DestinationFolderPath);
+			})
+			.OnTextChanged_Lambda([this](const FText& Text)
+			{
+				DestinationFolderPath = Text.ToString();
+			})
+			.OnTextCommitted_Lambda([this](const FText& Text, const ETextCommit::Type CommitType)
+			{
+				DestinationFolderPath = Text.ToString();
+			})
+		]
+		+ SHorizontalBox::Slot()
+		.AutoWidth()
+		.HAlign(HAlign_Right)
+		[
+			SNew(SButton)
+			.OnClicked_Lambda([this]()
+			{
+				OnOpenDestinationDirectoryClicked();
+				return FReply::Handled();
+			})
+			.ButtonStyle(FAppStyle::Get(), "SimpleButton")
+			[
+				SNew(SImage)
+				.Image(FAppStyle::Get().GetBrush("Zen.BrowseContent"))
+				.ColorAndOpacity(FSlateColor::UseForeground())
+			]
+		]
+	]
+
+	+ SVerticalBox::Slot()
+	.AutoHeight()
+	.VAlign(VAlign_Top)
+	[
+		SNew(SHorizontalBox)
+		.Visibility_Lambda([this]()
+		{
+			TArray<FBuildSelectionBuildGroupPtr> SelectedItems = BuildListView->GetSelectedItems();
+			return (SelectedItems.Num() == 0) || GetSelectedBuildType() != EBuildType::Oplog ? EVisibility::Collapsed : EVisibility::Visible;
+		})
+		+ SHorizontalBox::Slot()
+		.FillWidth(1.0f)
+		.VAlign(VAlign_Center)
+		[
+			SNew(SEditableTextBox)
+			.OverflowPolicy(ETextOverflowPolicy::MiddleEllipsis)
+			.MinDesiredWidth(200.0f)
+			.Text_Lambda([this]()
+			{
+				return FText::FromString(DestinationZenProjectId);
+			})
+			.OnTextChanged_Lambda([this](const FText& Text)
+			{
+				DestinationZenProjectId = Text.ToString();
+			})
+			.OnTextCommitted_Lambda([this](const FText& Text, const ETextCommit::Type CommitType)
+			{
+				DestinationZenProjectId = Text.ToString();
+			})
+		]
+	];
 }
 
 TSharedRef<SWidget> SBuildSelection::GetGridPanel()
@@ -389,9 +822,8 @@ TSharedRef<SWidget> SBuildSelection::GetGridPanel()
 	});
 
 	const float MinDesiredWidth = 50.0f;
-	int32 Row = 0;
 
-	const float RowMargin = 0.0f;
+	const float RowMargin = 2.0f;
 	const float ColumnMargin = 10.0f;
 	const FSlateColor TitleColor = FStyleColors::AccentWhite;
 	const FSlateFontInfo TitleFont = FCoreStyle::GetDefaultFontStyle("Bold", 10);
@@ -403,7 +835,7 @@ TSharedRef<SWidget> SBuildSelection::GetGridPanel()
 	[
 		SNew(SHorizontalBox)
 
-		// Project
+		// Stream
 		+SHorizontalBox::Slot()
 		.AutoWidth()
 		.FillWidth(1.0f)
@@ -420,7 +852,7 @@ TSharedRef<SWidget> SBuildSelection::GetGridPanel()
 				SNew(STextBlock)
 				.ColorAndOpacity(TitleColor)
 				.Font(TitleFont)
-				.Text(LOCTEXT("BuildFilter_Project", "Project"))
+				.Text(LOCTEXT("BuildSelection_Stream", "Stream"))
 			]
 
 			+ SVerticalBox::Slot()
@@ -429,15 +861,13 @@ TSharedRef<SWidget> SBuildSelection::GetGridPanel()
 			.HAlign(HAlign_Fill)
 			.VAlign(VAlign_Top)
 			[
-				SAssignNew(ProjectWidget, SComboBox<TSharedPtr<FString>>)
-				.OptionsSource(&ProjectList)
+				SAssignNew(StreamWidget, SComboBox<TSharedPtr<FString>>)
+				.OptionsSource(&StreamList)
 				.OnSelectionChanged_Lambda([this](TSharedPtr<FString> Item, ESelectInfo::Type SelectInfo)
 				{
-					if (Item.IsValid() && *Item != *SelectedProject)
+					if (Item.IsValid() && *Item != *SelectedStream)
 					{
-						SelectedProject = Item;
-						SelectedStream = nullptr;
-						SelectedBuildType = nullptr;
+						SelectedStream = Item;
 						RebuildLists();
 					}
 				})
@@ -447,13 +877,13 @@ TSharedRef<SWidget> SBuildSelection::GetGridPanel()
 						.MinDesiredWidth(MinDesiredWidth)
 						.Text_Lambda([this]()
 						{
-							return FText::FromString(SelectedProject ? **SelectedProject : TEXT(""));
+							return FText::FromString(SelectedStream ? **SelectedStream : TEXT(""));
 						})
 				]
 			]
 		]
 
-		// Stream
+		// Project
 		+SHorizontalBox::Slot()
 		.AutoWidth()
 		.FillWidth(1.0f)
@@ -470,7 +900,7 @@ TSharedRef<SWidget> SBuildSelection::GetGridPanel()
 				SNew(STextBlock)
 				.ColorAndOpacity(TitleColor)
 				.Font(TitleFont)
-				.Text(LOCTEXT("BuildFilter_Stream", "Stream"))
+				.Text(LOCTEXT("BuildSelection_Project", "Project"))
 			]
 
 			+ SVerticalBox::Slot()
@@ -479,14 +909,13 @@ TSharedRef<SWidget> SBuildSelection::GetGridPanel()
 			.HAlign(HAlign_Fill)
 			.VAlign(VAlign_Top)
 			[
-				SAssignNew(StreamWidget, SComboBox<TSharedPtr<FString>>)
-				.OptionsSource(&StreamList)
+				SAssignNew(ProjectWidget, SComboBox<TSharedPtr<FString>>)
+				.OptionsSource(&ProjectList)
 				.OnSelectionChanged_Lambda([this](TSharedPtr<FString> Item, ESelectInfo::Type SelectInfo)
 				{
-					if (Item.IsValid() && *Item != *SelectedStream)
+					if (Item.IsValid() && *Item != *SelectedProject)
 					{
-						SelectedStream = Item;
-						SelectedBuildType = nullptr;
+						SelectedProject = Item;
 						RebuildLists();
 					}
 				})
@@ -496,7 +925,7 @@ TSharedRef<SWidget> SBuildSelection::GetGridPanel()
 						.MinDesiredWidth(MinDesiredWidth)
 						.Text_Lambda([this]()
 						{
-							return FText::FromString(SelectedStream ? **SelectedStream : TEXT(""));
+							return FText::FromString(SelectedProject ? **SelectedProject : TEXT(""));
 						})
 				]
 			]
@@ -519,7 +948,7 @@ TSharedRef<SWidget> SBuildSelection::GetGridPanel()
 				SNew(STextBlock)
 				.ColorAndOpacity(TitleColor)
 				.Font(TitleFont)
-				.Text(LOCTEXT("BuildFilter_BuildType", "Build Type"))
+				.Text(LOCTEXT("BuildSelection_BuildType", "Build Type"))
 			]
 
 			+ SVerticalBox::Slot()
@@ -567,7 +996,7 @@ TSharedRef<SWidget> SBuildSelection::GetGridPanel()
 				SNew(STextBlock)
 				.ColorAndOpacity(TitleColor)
 				.Font(TitleFont)
-				.Text(LOCTEXT("BuildFilter_RequiredPlatforms", "Required Platforms"))
+				.Text(LOCTEXT("BuildSelection_RequiredPlatforms", "Required Platforms"))
 			]
 
 			+ SVerticalBox::Slot()
@@ -578,15 +1007,15 @@ TSharedRef<SWidget> SBuildSelection::GetGridPanel()
 			[
 				SAssignNew(RequiredPlatformsWidget, SMultiSelectComboBox)
 				.SelectValues(&PlatformList)
-				.OnCheckboxChanged_Lambda([this](ECheckBoxState)
+				.OnCheckedValuesChanged_Lambda([this]()
 				{
+					RegenerateActivePlatformFilters();
+					ValidateBuildGroupSelection();
 				})
 			]
 		]
 	];
 
-	++Row;
-
 	Panel->AddSlot()
 	.Padding(FMargin(ColumnMargin, 10, ColumnMargin, 0))
 	.AutoHeight()
@@ -596,11 +1025,9 @@ TSharedRef<SWidget> SBuildSelection::GetGridPanel()
 		SNew(STextBlock)
 			.ColorAndOpacity(TitleColor)
 			.Font(TitleFont)
-			.Text(LOCTEXT("BuildFilter_BuildsLabel", "Builds"))
+			.Text(LOCTEXT("BuildSelection_BuildsLabel", "Builds"))
 	];
 
-	++Row;
-
 	Panel->AddSlot()
 	.Padding(FMargin(ColumnMargin, RowMargin))
 	.HAlign(HAlign_Fill)
@@ -609,11 +1036,14 @@ TSharedRef<SWidget> SBuildSelection::GetGridPanel()
 		SNew(SVerticalBox)
 		+SVerticalBox::Slot()
 		.HAlign(HAlign_Fill)
+		.VAlign(VAlign_Fill)
 		[
 			SAssignNew(BuildListView, SListView<FBuildSelectionBuildGroupPtr>)
 				.ListItemsSource(&BuildGroups)
 				.OnGenerateRow(this, &SBuildSelection::GenerateBuildGroupRow)
+				.OnSelectionChanged(this, &SBuildSelection::BuildGroupSelectionChanged)
 				.SelectionMode(ESelectionMode::Single)
+				.OnIsSelectableOrNavigable(this, &SBuildSelection::BuildGroupIsSelectableOrNavigable)
 				.IsEnabled_Lambda([this]
 				{
 					return !BuildListRefreshesInProgress;
@@ -621,22 +1051,24 @@ TSharedRef<SWidget> SBuildSelection::GetGridPanel()
 				.HeaderRow
 				(
 					SNew(SHeaderRow)
-					+ SHeaderRow::Column(FBuildGroupIds::ColName).DefaultLabel(LOCTEXT("BuildFilter_BuildGroupColName", "Name"))
-					+ SHeaderRow::Column(FBuildGroupIds::ColCommit).DefaultLabel(LOCTEXT("BuildFilter_BuildGroupColCommit", "Commit"))
-						.DefaultTooltip(LOCTEXT("BuildFilter_BuildGroupColCommitTooltip", "Commit/Changelist for the build"))
-						.FixedWidth(60).HAlignCell(HAlign_Right).HAlignHeader(HAlign_Center).VAlignCell(VAlign_Center)
-					+ SHeaderRow::Column(FBuildGroupIds::ColSuffix).DefaultLabel(LOCTEXT("BuildFilter_BuildGroupColSuffix", "Suffix"))
-						.DefaultTooltip(LOCTEXT("BuildFilter_BuildGroupColSuffixTooltip", "Modifier on top of the commit/changelist for the build"))
-						.FixedWidth(60).HAlignCell(HAlign_Right).HAlignHeader(HAlign_Center).VAlignCell(VAlign_Center)
-					+ SHeaderRow::Column(FBuildGroupIds::ColCreated).DefaultLabel(LOCTEXT("BuildFilter_BuildGroupColCreated", "Created"))
-						.DefaultTooltip(LOCTEXT("BuildFilter_BuildGroupColCreatedTooltip", "When the build was created"))
-						.FixedWidth(60).HAlignCell(HAlign_Right).HAlignHeader(HAlign_Center).VAlignCell(VAlign_Center)
+					+ SHeaderRow::Column(FBuildGroupIds::ColName).DefaultLabel(LOCTEXT("BuildSelection_BuildGroupColName", "Name"))
+						.FillWidth(0.4f)
+					+ SHeaderRow::Column(FBuildGroupIds::ColCommit).DefaultLabel(LOCTEXT("BuildSelection_BuildGroupColCommit", "Commit"))
+						.DefaultTooltip(LOCTEXT("BuildSelection_BuildGroupColCommitTooltip", "Commit/Changelist for the build"))
+						.FillWidth(0.10f).HAlignCell(HAlign_Center).HAlignHeader(HAlign_Center).VAlignCell(VAlign_Center)
+					+ SHeaderRow::Column(FBuildGroupIds::ColSuffix).DefaultLabel(LOCTEXT("BuildSelection_BuildGroupColSuffix", "Suffix"))
+						.DefaultTooltip(LOCTEXT("BuildSelection_BuildGroupColSuffixTooltip", "Modifier on top of the commit/changelist for the build"))
+						.FillWidth(0.10f).HAlignCell(HAlign_Center).HAlignHeader(HAlign_Center).VAlignCell(VAlign_Center)
+					+ SHeaderRow::Column(FBuildGroupIds::ColCategory).DefaultLabel(LOCTEXT("BuildSelection_BuildGroupColCategory", "Category"))
+						.DefaultTooltip(LOCTEXT("BuildSelection_BuildGroupColCategoryTooltip", "Category for the build"))
+						.FillWidth(0.25f).HAlignCell(HAlign_Left).HAlignHeader(HAlign_Center).VAlignCell(VAlign_Center)
+					+ SHeaderRow::Column(FBuildGroupIds::ColCreated).DefaultLabel(LOCTEXT("BuildSelection_BuildGroupColCreated", "Created"))
+						.DefaultTooltip(LOCTEXT("BuildSelection_BuildGroupColCreatedTooltip", "When the build was created"))
+						.FillWidth(0.15f).HAlignCell(HAlign_Left).HAlignHeader(HAlign_Center).VAlignCell(VAlign_Center)
 				)
 		]
 	];
 
-	++Row;
-
 	Panel->AddSlot()
 	.Padding(FMargin(ColumnMargin, 3, ColumnMargin, 0))
 	.AutoHeight()
@@ -649,15 +1081,248 @@ TSharedRef<SWidget> SBuildSelection::GetGridPanel()
 			{
 				if (BuildListRefreshesInProgress)
 				{
-					return LOCTEXT("BuildFilter_ResultLoading", "Loading");
+					return LOCTEXT("BuildSelection_ResultLoading", "Loading...");
+				}
+				int32 VisibleItemCount = 0;
+				if (ActivePlatformFilters.IsEmpty())
+				{
+					VisibleItemCount = BuildGroups.Num();
+				}
+				else
+				{
+					for (FBuildSelectionBuildGroupPtr BuildGroup : BuildGroups)
+					{
+						bool bHasAllRequiredPlatforms = true;
+						for (const FString& ActivePlatformFilter : ActivePlatformFilters)
+						{
+							if (!BuildGroup->PerPlatformBuilds.Contains(ActivePlatformFilter))
+							{
+								bHasAllRequiredPlatforms = false;
+								break;
+							}
+						}
+						if (bHasAllRequiredPlatforms)
+						{
+							VisibleItemCount++;
+						}
+					}
+				}
+				if (VisibleItemCount == 1)
+				{
+					return LOCTEXT("BuildSelection_ResultDescriptionMultiple", "1 item");
+				}
+				else
+				{
+					return FText::Format(LOCTEXT("BuildSelection_ResultDescriptionMultiple", "{0} items"), FText::AsNumber(VisibleItemCount));
 				}
-				return FText::Format(LOCTEXT("BuildFilter_ResultDescription", "{0} items"), FText::AsNumber(BuildGroups.Num()));
 			})
 	];
 
-	++Row;
+	Panel->AddSlot()
+	.Padding(FMargin(ColumnMargin, RowMargin))
+	.AutoHeight()
+	.HAlign(HAlign_Fill)
+	.VAlign(VAlign_Bottom)
+	[
+
+		SNew(SHorizontalBox)
+		.Visibility_Lambda([this]()
+		{
+			return !!BuildListRefreshesInProgress || BuildListView->GetNumItemsSelected() == 0 ? EVisibility::Collapsed : EVisibility::Visible;
+		})
+		+SHorizontalBox::Slot()
+		.HAlign(HAlign_Left)
+		.FillWidth(0.5f)
+		[
+			SNew(SVerticalBox)
+			+SVerticalBox::Slot()
+			.HAlign(HAlign_Fill)
+			.VAlign(VAlign_Top)
+			.AutoHeight()
+			[
+				SNew(STextBlock)
+					.ColorAndOpacity(TitleColor)
+					.Font(TitleFont)
+					.Text(LOCTEXT("BuildSelection_AvailablePlatforms", "Available Platforms"))
+			]
+			+SVerticalBox::Slot()
+			.HAlign(HAlign_Fill)
+			.VAlign(VAlign_Top)
+			.AutoHeight()
+			[
+				SAssignNew(SelectedGroupPlatformGrid, SGridPanel)
+			]
+		]
+		+SHorizontalBox::Slot()
+		.HAlign(HAlign_Right)
+		.FillWidth(0.5f)
+		[
+			SNew(SVerticalBox)
+			+SVerticalBox::Slot()
+			.HAlign(HAlign_Fill)
+			.VAlign(VAlign_Top)
+			.AutoHeight()
+			[
+				SNew(STextBlock)
+					.ColorAndOpacity(TitleColor)
+					.Font(TitleFont)
+					.Text(LOCTEXT("BuildSelection_Destination", "Destination"))
+			]
+			+SVerticalBox::Slot()
+			.HAlign(HAlign_Fill)
+			.VAlign(VAlign_Top)
+			[
+				GetBuildDestinationPanel()
+			]
+			+SVerticalBox::Slot()
+			.HAlign(HAlign_Right)
+			.VAlign(VAlign_Bottom)
+			.Padding(FMargin(0, RowMargin))
+			[
+				SNew(SButton)
+				.Text(LOCTEXT("BuildSelection_Download", "Download"))
+				.ToolTipText(LOCTEXT("BuildSelection_DownloadTooltip", "Start a download of the selected build for the selected platforms"))
+				.ButtonStyle(FAppStyle::Get(), "Button")
+				.IsEnabled_Lambda([this]()
+				{
+					const bool bDestinationValid = GetSelectedBuildType() == EBuildType::Oplog ? !DestinationZenProjectId.IsEmpty() : !DestinationFolderPath.IsEmpty();
+					if (!bDestinationValid)
+					{
+						return false;
+					}
+
+					TArray<FBuildSelectionBuildGroupPtr> SelectedItems = BuildListView->GetSelectedItems();
+					if (SelectedItems.Num() == 0)
+					{
+						return false;
+					}
+
+					for (const FString& PlatformForSelectedGroup : SelectedGroupSelectedPlatforms)
+					{
+						if (SelectedItems[0]->PerPlatformBuilds.Contains(PlatformForSelectedGroup))
+						{
+							return true;
+						}
+					}
+					return false;
+				})
+				.OnClicked_Lambda([this]()
+				{
+					using namespace UE::Zen::Build;
+
+					if (TSharedPtr<FBuildServiceInstance> ServiceInstance = BuildServiceInstance.Get())
+					{
+						TArray<FBuildSelectionBuildGroupPtr> SelectedItems = BuildListView->GetSelectedItems();
+						if (SelectedItems.Num() == 0)
+						{
+							return FReply::Handled();
+						}
+
+						for (const FString& PlatformForSelectedGroup : SelectedGroupSelectedPlatforms)
+						{
+							if (FBuildServiceInstance::FBuildRecord* BuildRecord = SelectedItems[0]->PerPlatformBuilds.Find(PlatformForSelectedGroup))
+							{
+								FString Bucket = FString::Printf(TEXT("%s.%s.%s.%s"), **SelectedProject, **SelectedBuildType, **SelectedStream, *PlatformForSelectedGroup);
+								if (GetSelectedBuildType() == EBuildType::Oplog)
+								{
+									if (FCbFieldView CookPlatformField = BuildRecord->Metadata["cookPlatform"]; CookPlatformField.HasValue() && !CookPlatformField.HasError())
+									{
+										FString ProjectFilePath = FUProjectDictionary::GetDefault().GetProjectPathForGame(**SelectedProject);
+										FString DestinationOplogId = *WriteToString<64>(CookPlatformField.AsString());
+										FBuildServiceInstance::FBuildTransfer BuildTransfer =
+											ServiceInstance->StartOplogBuildTransfer(BuildRecord->BuildId, DestinationZenProjectId, DestinationOplogId, ProjectFilePath, SelectedItems[0]->Namespace, Bucket);
+										OnBuildTransferStarted.ExecuteIfBound(BuildTransfer, SelectedItems[0]->DisplayName, PlatformForSelectedGroup);
+									}
+								}
+								else
+								{
+									FString DestinationFolder = FPaths::Combine(DestinationFolderPath, PlatformForSelectedGroup);
+									FBuildServiceInstance::FBuildTransfer BuildTransfer =
+										ServiceInstance->StartBuildTransfer(BuildRecord->BuildId, DestinationFolder, SelectedItems[0]->Namespace, Bucket);
+									OnBuildTransferStarted.ExecuteIfBound(BuildTransfer, SelectedItems[0]->DisplayName, PlatformForSelectedGroup);
+								}
+							}
+						}
+					}
+					return FReply::Handled();
+				})
+			]
+		]
+	];
 
 	return Panel;
 }
 
+void
+SBuildGroupTableRow::Construct(const FArguments& InArgs, const TSharedRef<STableViewBase>& InOwnerTableView, const FBuildSelectionBuildGroupPtr InBuildGroup)
+{
+	BuildGroup = InBuildGroup;
+
+	SMultiColumnTableRow<FBuildSelectionBuildGroupPtr>::Construct(FSuperRowType::FArguments(), InOwnerTableView);
+}
+
+TSharedRef<SWidget>
+SBuildGroupTableRow::GenerateWidgetForColumn(const FName& ColumnName)
+{
+	using namespace UE::BuildSelection::Internal;
+
+	if (ColumnName == FBuildGroupIds::ColName)
+	{
+		return SNew(SHorizontalBox)
+				+ SHorizontalBox::Slot()
+				[
+					SNew(STextBlock).Text(FText::FromString(BuildGroup->DisplayName))
+				];
+	}
+	else if (ColumnName == FBuildGroupIds::ColCommit)
+	{
+		return SNew(SHorizontalBox)
+				+ SHorizontalBox::Slot()
+				[
+					SNew(STextBlock).Text(FText::FromString(BuildGroup->CommitIdentifier))
+				];
+	}
+	else if (ColumnName == FBuildGroupIds::ColSuffix)
+	{
+		return SNew(SHorizontalBox)
+				+ SHorizontalBox::Slot()
+				[
+					SNew(STextBlock).Text(FText::FromString(BuildGroup->Suffix))
+				];
+	}
+	else if (ColumnName == FBuildGroupIds::ColCategory)
+	{
+		return SNew(SHorizontalBox)
+				+ SHorizontalBox::Slot()
+				[
+					SNew(STextBlock).Text(FText::FromString(BuildGroup->Category))
+				];
+	}
+	else if (ColumnName == FBuildGroupIds::ColCreated)
+	{
+		if (BuildGroup->CreatedAt.GetTicks() != 0)
+		{
+			return SNew(SHorizontalBox)
+				+ SHorizontalBox::Slot()
+				[
+					SNew(STextBlock).Text(FText::AsDateTime(BuildGroup->CreatedAt, EDateTimeStyle::Short))
+				];
+		}
+	}
+
+	return SNullWidget::NullWidget;
+}
+
+const FSlateBrush*
+SBuildGroupTableRow::GetBorder() const
+{
+	return STableRow<FBuildSelectionBuildGroupPtr>::GetBorder();
+}
+
+FReply
+SBuildGroupTableRow::OnBrowseClicked()
+{
+	return FReply::Unhandled();
+}
+
 #undef LOCTEXT_NAMESPACE
diff --git a/Engine/Source/Editor/StorageServerWidgets/Private/SMultiSelectComboBox.cpp b/Engine/Source/Editor/StorageServerWidgets/Private/SMultiSelectComboBox.cpp
--- a/Engine/Source/Editor/StorageServerWidgets/Private/SMultiSelectComboBox.cpp
+++ b/Engine/Source/Editor/StorageServerWidgets/Private/SMultiSelectComboBox.cpp
@@ -15,6 +15,7 @@
 void SMultiSelectComboBox::Construct(const FArguments& InArgs)
 {
 	using namespace UE::Slate::Containers;
+	OnCheckedValuesChanged = InArgs._OnCheckedValuesChanged;
 	SelectValues = InArgs._SelectValues;
 
 	TSharedPtr<SHorizontalBox> HorizontalBox;
@@ -70,6 +71,11 @@ SMultiSelectComboBox::~SMultiSelectComboBox()
 {
 }
 
+bool SMultiSelectComboBox::IsChecked(FStringView Value) const
+{
+	return CheckedValues.Contains(Value);
+}
+
 TSharedRef<class SWidget> SMultiSelectComboBox::GenerateMenuItems()
 {
 	TSharedRef<SVerticalBox> Contents = SNew(SVerticalBox);
@@ -124,12 +130,20 @@ void SMultiSelectComboBox::OnSelectedChangedFromMultiselect(TSharedPtr<FString>
 	if(!CheckedValues.Contains(*Value))
 	{
 		CheckedValues.Add(FString(*Value));
+		OnCheckedValuesChanged.Execute();
 	}
 	else
 	{
+		bool bAnyRemoved = false;
 		while (CheckedValues.Contains(*Value))
 		{
 			CheckedValues.Remove(FString(*Value));
+			bAnyRemoved = true;
+		}
+
+		if (bAnyRemoved)
+		{
+			OnCheckedValuesChanged.Execute();
 		}
 	}
 }
diff --git a/Engine/Source/Editor/StorageServerWidgets/Private/SMultiSelectComboBox.h b/Engine/Source/Editor/StorageServerWidgets/Private/SMultiSelectComboBox.h
--- a/Engine/Source/Editor/StorageServerWidgets/Private/SMultiSelectComboBox.h
+++ b/Engine/Source/Editor/StorageServerWidgets/Private/SMultiSelectComboBox.h
@@ -6,27 +6,28 @@
 #include "Styling/SlateTypes.h"
 #include "Containers/ObservableArray.h"
 
-DECLARE_DELEGATE_OneParam(FOnCheckboxChangedSignature, ECheckBoxState)
+DECLARE_DELEGATE(FOnMultiSelectComboBoxCheckedValuesChanged)
 
 class SMultiSelectComboBox : public SCompoundWidget
 {
 public:
 	SLATE_BEGIN_ARGS(SMultiSelectComboBox) {}
 		SLATE_ARGUMENT(TArray<TSharedPtr<FString>>*, SelectValues)
-		SLATE_EVENT(FOnCheckboxChangedSignature, OnCheckboxChanged)
+		SLATE_EVENT(FOnMultiSelectComboBoxCheckedValuesChanged, OnCheckedValuesChanged)
 	SLATE_END_ARGS()
 
 	virtual ~SMultiSelectComboBox();
 
 	void Construct(const FArguments& InArgs);
 
+	bool IsChecked(FStringView Value) const;
+
 private:
 	TSharedRef<class SWidget> GenerateMenuItems();
 
 	void OnSelectedChangedFromMultiselect(TSharedPtr<FString> Value);
 
-	FReply OnLabelClick();
-
+	FOnMultiSelectComboBoxCheckedValuesChanged OnCheckedValuesChanged;
 	TArray<FString> CheckedValues;
 	TArray<TSharedPtr<FString>>* SelectValues;
 };
\ No newline at end of file
diff --git a/Engine/Source/Editor/StorageServerWidgets/StorageServerWidgets.Build.cs b/Engine/Source/Editor/StorageServerWidgets/StorageServerWidgets.Build.cs
--- a/Engine/Source/Editor/StorageServerWidgets/StorageServerWidgets.Build.cs
+++ b/Engine/Source/Editor/StorageServerWidgets/StorageServerWidgets.Build.cs
@@ -12,13 +12,13 @@ public class StorageServerWidgets : ModuleRules
 				"ApplicationCore",
 				"Core",
 				"CoreUObject",
+				"DesktopPlatform",
 				"InputCore",
 				"Json",
 				"Projects",	
 				"Slate",
 				"SlateCore",
-				"Zen",
-				"DesktopPlatform"
+				"Zen"
 			});
 
 		CppCompileWarningSettings.UnsafeTypeCastWarningLevel = WarningLevel.Error;
diff --git a/Engine/Source/Programs/BuildStorageTool/Private/BuildStorageTool.cpp b/Engine/Source/Programs/BuildStorageTool/Private/BuildStorageTool.cpp
--- a/Engine/Source/Programs/BuildStorageTool/Private/BuildStorageTool.cpp
+++ b/Engine/Source/Programs/BuildStorageTool/Private/BuildStorageTool.cpp
@@ -152,7 +152,7 @@ public:
 		Window =
 			SNew(SWindow)
 			.Title(GetWindowTitle())
-			.ClientSize(FVector2D(700.0f, 400.0f))
+			.ClientSize(FVector2D(900.0f, 600.0f))
 			.ActivationPolicy(EWindowActivationPolicy::Always)
 			.SizingRule(ESizingRule::UserSized)
 			.IsTopmostWindow(false)
@@ -176,7 +176,7 @@ public:
 					// Login panel
 					+ SVerticalBox::Slot()
 					.AutoHeight()
-					.Padding(0.0f, 10.0f, 0.0f, 0.0f)
+					.Padding(0.0f, 10.0f, 5.0f, 0.0f)
 					.HAlign(HAlign_Fill)
 					.VAlign(VAlign_Top)
 					[
@@ -187,23 +187,28 @@ public:
 
 					// Selection panel
 					+ SVerticalBox::Slot()
-					.Padding(0.0f, 10.0f, 0.0f, 0.0f)
+					.Padding(0.0f, 10.0f, 5.0f, 0.0f)
 					.HAlign(HAlign_Fill)
 					.VAlign(VAlign_Top)
 					[
 						SNew(SBuildSelection)
 						.ZenServiceInstance(ZenServiceInstanceManager.ToSharedRef(), &UE::Zen::FServiceInstanceManager::GetZenServiceInstance)
 						.BuildServiceInstance(BuildServiceInstanceManager.ToSharedRef(), &UE::Zen::Build::FServiceInstanceManager::GetBuildServiceInstance)
+						.OnBuildTransferStarted_Lambda([this]
+							(UE::Zen::Build::FBuildServiceInstance::FBuildTransfer Transfer, FStringView Name, FStringView Platform)
+							{
+								BuildActivity->AddBuildTransfer(Transfer, Name, Platform);
+							})
 					]
 
 					// Activity panel
 					+ SVerticalBox::Slot()
 					.AutoHeight()
-					.Padding(0.0f, 10.0f, 0.0f, 0.0f)
+					.Padding(0.0f, 10.0f, 5.0f, 0.0f)
 					.HAlign(HAlign_Fill)
 					.VAlign(VAlign_Bottom)
 					[
-						SNew(SBuildActivity)
+						SAssignNew(BuildActivity, SBuildActivity)
 						.ZenServiceInstance(ZenServiceInstanceManager.ToSharedRef(), &UE::Zen::FServiceInstanceManager::GetZenServiceInstance)
 						.BuildServiceInstance(BuildServiceInstanceManager.ToSharedRef(), &UE::Zen::Build::FServiceInstanceManager::GetBuildServiceInstance)
 					]
@@ -259,6 +264,8 @@ private:
 	{
 		FCoreDelegates::ModalMessageDialog.Unbind();
 	}
+
+	TSharedPtr<SBuildActivity> BuildActivity;
 };
 
 int BuildStorageToolMain(const TCHAR* CmdLine)
@@ -271,6 +278,8 @@ int BuildStorageToolMain(const TCHAR* CmdLine)
 		FEngineLoop::AppExit();
 	};
 
+	const FTaskTagScope Scope(ETaskTag::EGameThread);
+
 	// start up the main loop
 	GEngineLoop.PreInit(CmdLine);
 
diff --git a/Engine/Source/Programs/BuildStorageTool/Private/BuildStorageToolStyle.cpp b/Engine/Source/Programs/BuildStorageTool/Private/BuildStorageToolStyle.cpp
--- a/Engine/Source/Programs/BuildStorageTool/Private/BuildStorageToolStyle.cpp
+++ b/Engine/Source/Programs/BuildStorageTool/Private/BuildStorageToolStyle.cpp
@@ -51,12 +51,16 @@ TSharedRef< FSlateStyleSet > FBuildStorageToolStyle::Create()
 
 	Style.SetParentStyleName("CoreStyle");
 
+	const ISlateStyle* ParentStyle = FSlateStyleRegistry::FindSlateStyle("CoreStyle");
+	const FTextBlockStyle NormalText = ParentStyle->GetWidgetStyle<FTextBlockStyle>("NormalText");
+
 	const FVector2D Icon12x12(12.0f, 12.0f);
 	const FVector2D Icon16x16(16.0f, 16.0f);
 
 	Style.Set( "AppIcon", new IMAGE_BRUSH( "Icons/DefaultAppIcon", FVector2D(20,20)) );
 
 	Style.Set("Zen.FolderExplore", new EDITOR_IMAGE_BRUSH_SVG("Starship/Common/ContentBrowser", Icon12x12));
+	Style.Set("Zen.BrowseContent", new EDITOR_IMAGE_BRUSH_SVG("Starship/Common/BrowseContent", Icon16x16));
 	Style.Set("Zen.Clipboard", new EDITOR_IMAGE_BRUSH_SVG("Starship/Common/Clipboard", Icon16x16));
 
 	const FTextBlockStyle DefaultText = FTextBlockStyle()
@@ -79,6 +83,26 @@ TSharedRef< FSlateStyleSet > FBuildStorageToolStyle::Create()
 		.SetColorAndOpacity(FSlateColor::UseSubduedForeground())
 	);
 
+	{
+		// Navigation defaults
+		const FLinearColor NavHyperlinkColor(0.03847f, 0.33446f, 1.0f);
+		const FTextBlockStyle NavigationHyperlinkText = FTextBlockStyle(NormalText)
+			.SetFont(DEFAULT_FONT("Regular", 10))
+			.SetColorAndOpacity(NavHyperlinkColor);
+
+		const FButtonStyle NavigationHyperlinkButton = FButtonStyle()
+			.SetNormal(BORDER_BRUSH("Old/HyperlinkDotted", FMargin(0, 0, 0, 3 / 16.0f), NavHyperlinkColor))
+			.SetPressed(FSlateNoResource())
+			.SetHovered(BORDER_BRUSH("Old/HyperlinkUnderline", FMargin(0, 0, 0, 3 / 16.0f), NavHyperlinkColor));
+
+		FHyperlinkStyle NavigationHyperlink = FHyperlinkStyle()
+			.SetUnderlineStyle(NavigationHyperlinkButton)
+			.SetTextStyle(NavigationHyperlinkText)
+			.SetPadding(FMargin(0.0f));
+
+		Style.Set("NavigationHyperlink", NavigationHyperlink);
+	}
+
 
 	FSlateBrush* GenericWhiteBox = new IMAGE_BRUSH( "Old/White", Icon16x16 );
 
diff --git a/Engine/Source/Runtime/Core/Private/Misc/App.cpp b/Engine/Source/Runtime/Core/Private/Misc/App.cpp
--- a/Engine/Source/Runtime/Core/Private/Misc/App.cpp
+++ b/Engine/Source/Runtime/Core/Private/Misc/App.cpp
@@ -472,21 +472,7 @@ FString FApp::GetZenStoreProjectId(FStringView SubProject)
 #if PLATFORM_DESKTOP
 	if (FPaths::IsProjectFilePathSet())
 	{
-		FString ProjectFilePath = FPaths::GetProjectFilePath();
-		FString AbsProjectFilePath = IFileManager::Get().ConvertToAbsolutePathForExternalAppForRead(*ProjectFilePath);
-		AbsProjectFilePath = FPaths::FindCorrectCase(AbsProjectFilePath);
-		FTCHARToUTF8 AbsProjectFilePathUTF8(*AbsProjectFilePath);
-
-		FString HashString = FMD5::HashBytes((unsigned char*)AbsProjectFilePathUTF8.Get(), AbsProjectFilePathUTF8.Length()).Left(8);
-
-		if (SubProject.IsEmpty())
-		{
-			return FString::Printf(TEXT("%s.%.8s"), FApp::GetProjectName(), *HashString);
-		}
-		else
-		{
-			return FString::Printf(TEXT("%s.%.*s.%.8s"), FApp::GetProjectName(), SubProject.Len(), SubProject.GetData(), *HashString);
-		}
+		return GetZenStoreProjectIdForProject(FPaths::GetProjectFilePath());
 	}
 	UE_LOG(LogInit, Fatal, TEXT("GetZenStoreProjectId() called before having a valid project file path"));
 #else
@@ -495,6 +481,33 @@ FString FApp::GetZenStoreProjectId(FStringView SubProject)
 	return FString();
 }
 
+#if PLATFORM_DESKTOP
+FString FApp::GetZenStoreProjectIdForProject(FStringView ProjectFilePath, FStringView SubProject)
+{
+	if (ProjectFilePath.IsEmpty())
+	{
+		return FString();
+	}
+
+	FString ProjectFilePathStr(ProjectFilePath);
+	FString AbsProjectFilePath = IFileManager::Get().ConvertToAbsolutePathForExternalAppForRead(*ProjectFilePathStr);
+	AbsProjectFilePath = FPaths::FindCorrectCase(AbsProjectFilePath);
+	FTCHARToUTF8 AbsProjectFilePathUTF8(*AbsProjectFilePath);
+
+	FString HashString = FMD5::HashBytes((unsigned char*)AbsProjectFilePathUTF8.Get(), AbsProjectFilePathUTF8.Length()).Left(8);
+	FString ProjectName = FPaths::GetBaseFilename(ProjectFilePathStr);
+
+	if (SubProject.IsEmpty())
+	{
+		return FString::Printf(TEXT("%s.%.8s"), *ProjectName, *HashString);
+	}
+	else
+	{
+		return FString::Printf(TEXT("%s.%.*s.%.8s"), *ProjectName, SubProject.Len(), SubProject.GetData(), *HashString);
+	}
+}
+#endif
+
 
 FGuid FApp::GetInstanceId()
 {
diff --git a/Engine/Source/Runtime/Core/Public/Misc/App.h b/Engine/Source/Runtime/Core/Public/Misc/App.h
--- a/Engine/Source/Runtime/Core/Public/Misc/App.h
+++ b/Engine/Source/Runtime/Core/Public/Misc/App.h
@@ -239,12 +239,23 @@ public:
 	}
 
 	/**
-	 * Gets the Zen store project id for the current application instance.
+	 * Gets the Zen store project id for the current project file path.
 	 *
+	 * @param SubProject An optional subproject identifier (used for DLC).
 	 * @return Zen store project id.
 	 */
 	static CORE_API FString GetZenStoreProjectId(FStringView SubProject = FStringView());
 
+#if PLATFORM_DESKTOP
+	/**
+	 * Gets the Zen store project id for an explicitly specified project file path.
+	 *
+	 * @param SubProject An optional subproject identifier (used for DLC).
+	 * @return Zen store project id.
+	 */
+	static CORE_API FString GetZenStoreProjectIdForProject(FStringView ProjectFilePath, FStringView SubProject = FStringView());
+#endif
+
 	/**
 	 * Gets the globally unique identifier of this application instance.
 	 *
