UX improvements for Build Storage Tool

diff --git a/Engine/Source/Developer/IoStoreUtilities/Private/ZenStoreWriter.cpp b/Engine/Source/Developer/IoStoreUtilities/Private/ZenStoreWriter.cpp
--- a/Engine/Source/Developer/IoStoreUtilities/Private/ZenStoreWriter.cpp
+++ b/Engine/Source/Developer/IoStoreUtilities/Private/ZenStoreWriter.cpp
@@ -8,6 +8,7 @@
 #include "AssetRegistry/AssetRegistryState.h"
 #include "Async/Async.h"
 #include "Containers/Queue.h"
+#include "Experimental/ZenProjectStoreWriter.h"
 #include "GenericPlatform/GenericPlatformFile.h"
 #include "HAL/FileManager.h"
 #include "HAL/IConsoleManager.h"
@@ -480,86 +481,22 @@ void FZenStoreWriter::Initialize(const FCookInfo& Info)
 		}
 
 		FString OplogLifetimeMarkerPath = OutputPath / TEXT("ue.projectstore");
-		TUniquePtr<FArchive> OplogMarker(IFileManager::Get().CreateFileWriter(*OplogLifetimeMarkerPath));
+		Zen::FZenProjectStoreWriter ProjectStoreWriter(*OplogLifetimeMarkerPath);
 
 		bool bOplogEstablished = HttpClient->TryCreateOplog(ProjectId, OplogId, OplogLifetimeMarkerPath);
 
-		if (bOplogEstablished && OplogMarker)
-		{
-			FStringView ServiceHostName = TEXT("localhost");
-			uint16 ServicePort = 8558;
-			bool IsRunningLocally = false;
+		FStringView ServiceHostName = TEXT("localhost");
+		uint16 ServicePort = 8558;
+		bool IsRunningLocally = false;
 
 #if UE_WITH_ZEN
-			const Zen::FZenServiceInstance& ZenServiceInstance = HttpClient->GetZenServiceInstance();
-			IsRunningLocally = ZenServiceInstance.IsServiceRunningLocally();
-			ServiceHostName = ZenServiceInstance.GetEndpoint().GetHostName();
-			ServicePort = ZenServiceInstance.GetEndpoint().GetPort();
-#endif
-			TSharedRef<TJsonWriter<UTF8CHAR, TPrettyJsonPrintPolicy<UTF8CHAR>>> Writer = TJsonWriterFactory<UTF8CHAR, TPrettyJsonPrintPolicy<UTF8CHAR>>::Create(OplogMarker.Get());
-			Writer->WriteObjectStart();
-			Writer->WriteObjectStart(TEXT("zenserver"));
-			Writer->WriteValue(TEXT("islocalhost"), IsRunningLocally);
-			Writer->WriteValue(TEXT("hostname"), ServiceHostName);
-			if (IsRunningLocally)
-			{
-				ISocketSubsystem* SocketSubsystem = ISocketSubsystem::Get();
-				if (SocketSubsystem != nullptr)
-				{
-					TArray<TSharedPtr<FInternetAddr>> Addresses;
-					if (SocketSubsystem->GetLocalAdapterAddresses(Addresses))
-					{
-						Writer->WriteArrayStart("remotehostnames");
-						for (const TSharedPtr<FInternetAddr>& Address : Addresses)
-						{
-							Writer->WriteValue(Address->ToString(false));
-						}
-						FString MachineHostName;
-						if (SocketSubsystem->GetHostName(MachineHostName))
-						{
-							// Try to acquire FQDN
-							FAddressInfoResult GAIRequest = SocketSubsystem->GetAddressInfo(*MachineHostName, nullptr, EAddressInfoFlags::AllResultsWithMapping | EAddressInfoFlags::OnlyUsableAddresses | EAddressInfoFlags::AllResults | EAddressInfoFlags::FQDomainName, NAME_None);
-							if (GAIRequest.ReturnCode == SE_NO_ERROR && !GAIRequest.CanonicalNameResult.IsEmpty())
-							{
-								MachineHostName = GAIRequest.CanonicalNameResult;
-							}
-
-							Writer->WriteValue(TEXT("hostname://") + MachineHostName);
-						}
-#if PLATFORM_MAC
-						// Store the Bonjour hostname when on a mac so that the client can attempt to connect via USB.
-						// This is done since the "Link Local" IP of the USB connection changes everytime
-						// the cable is plugged in, so there is no stable IP to save.
-						bool bAppleTarget = TargetPlatformFName == TEXT("IOS") || TargetPlatformFName == TEXT("TVOS");
-						if (bAppleTarget)
-						{
-							FString MacBonjourName = FPlatformMisc::GetBonjourName();
-							
-							// If there is no Bonjour name, fallback to the mac's hostname
-							if (MacBonjourName.IsEmpty())
-							{
-								MacBonjourName = FPlatformProcess::ComputerName();
-							}
-							
-							if (!MacBonjourName.IsEmpty())
-							{
-								Writer->WriteValue(TEXT("macserver://") + MacBonjourName);
-							}
-						}
+		const Zen::FZenServiceInstance& ZenServiceInstance = HttpClient->GetZenServiceInstance();
+		IsRunningLocally = ZenServiceInstance.IsServiceRunningLocally();
+		ServiceHostName = ZenServiceInstance.GetEndpoint().GetHostName();
+		ServicePort = ZenServiceInstance.GetEndpoint().GetPort();
 #endif
-						Writer->WriteArrayEnd();
-					}
-				}
-			}
-			Writer->WriteValue(TEXT("hostport"), ServicePort);
-			Writer->WriteValue(TEXT("projectid"), ProjectId);
-			Writer->WriteValue(TEXT("oplogid"), OplogId);
-			Writer->WriteObjectEnd();
-			Writer->WriteObjectEnd();
-			Writer->Close();
-		}
 
-		OplogMarker.Reset();
+		ProjectStoreWriter.Write(ServiceHostName, ServicePort, IsRunningLocally, *ProjectId, *OplogId);
 
 		if (!bOplogEstablished && CleanBuild)
 		{
diff --git a/Engine/Source/Developer/Zen/Private/BuildServerInterface.cpp b/Engine/Source/Developer/Zen/Private/BuildServerInterface.cpp
--- a/Engine/Source/Developer/Zen/Private/BuildServerInterface.cpp
+++ b/Engine/Source/Developer/Zen/Private/BuildServerInterface.cpp
@@ -5,9 +5,11 @@
 #include "Async/Mutex.h"
 #include "Containers/StringView.h"
 #include "DesktopPlatformModule.h"
+#include "Experimental/ZenProjectStoreWriter.h"
 #include "Experimental/ZenServerInterface.h"
 #include "HAL/CriticalSection.h"
 #include "HAL/FileManager.h"
+#include "HAL/PlatformFileManager.h"
 #include "Http/HttpClient.h"
 #include "Http/HttpHostBuilder.h"
 #include "Misc/App.h"
@@ -23,6 +25,7 @@
 #include "Serialization/CompactBinarySerialization.h"
 #include "Serialization/CompactBinaryValidation.h"
 #include "Serialization/CompactBinaryWriter.h"
+#include "SocketSubsystem.h"
 #include "String/LexFromString.h"
 #include "Tasks/Task.h"
 
@@ -223,6 +226,7 @@ FBuildServiceInstance::RefreshNamespacesAndBuckets(FOnRefreshNamespacesAndBucket
 
 	FString CommandLineArgs = FString::Printf(TEXT("builds list-namespaces --recursive \"%s\""),
 		*OutputFilePath);
+	CommandLineArgs = AddZenUtilityBuildServerArguments(*CommandLineArgs);
 
 	InvokeZenUtility(CommandLineArgs, nullptr,
 	[this, OutputFilePath = MoveTemp(OutputFilePath), InOnRefreshNamespacesAndBucketsComplete = MoveTemp(InOnRefreshNamespacesAndBucketsComplete)](bool bSuccessful) mutable
@@ -274,6 +278,7 @@ FBuildServiceInstance::ListBuilds(FStringView Namespace, FStringView Bucket, FOn
 
 	FString CommandLineArgs = FString::Printf(TEXT("builds list --namespace \"%s\" --bucket \"%s\" \"%s\" \"%s\""),
 		*WriteToString<64>(Namespace), *WriteToString<64>(Bucket), *QueryFilePath, *OutputFilePath);
+	CommandLineArgs = AddZenUtilityBuildServerArguments(*CommandLineArgs);
 
 	auto PreInvoke = [QueryFilePath](FString& CommandLineArgs)
 	{
@@ -391,6 +396,18 @@ FBuildServiceInstance::FBuildTransfer::GetStatus() const
 	return BuildTransferCommand->Status;
 }
 
+void
+FBuildServiceInstance::FBuildTransfer::RequestCancel()
+{
+	if (!BuildTransferCommand)
+	{
+		return;
+	}
+	BuildTransferCommand->bCancelRequested.store(true);
+	FWriteScopeLock _(BuildTransferCommand->Lock);
+	BuildTransferCommand->Status = EBuildTransferStatus::Canceled;
+}
+
 FBuildServiceInstance::EConnectionState
 FBuildServiceInstance::GetConnectionState() const
 {
@@ -411,6 +428,7 @@ FBuildServiceInstance::Initialize()
 	// and doesn't unload SSL before unloading HTTP module at exit
 	FSslModule::Get();
 #endif
+	ISocketSubsystem::Get();
 	FDesktopPlatformModule::TryGet();
 }
 
@@ -539,80 +557,82 @@ FBuildServiceInstance::GetAccessToken() const
 	return FString(AccessTokenBuilder);
 }
 
-void
-FBuildServiceInstance::InvokeZenUtility(FStringView InCommandLineArgs, FPreZenUtilityInvocation&& PreInvocation, FOnZenUtilityInvocationComplete&& OnComplete)
+FString
+FBuildServiceInstance::AddZenUtilityBuildServerArguments(const TCHAR* InCommandLineArgs) const
 {
-	UE::Tasks::Launch(UE_SOURCE_LOCATION,
-	[this, InCommandLineArgs = FString(InCommandLineArgs), PreInvocation = MoveTemp(PreInvocation), OnComplete = MoveTemp(OnComplete)]() mutable
-	{
- 		FString ZenUtilityPath = UE::Zen::GetLocalInstallUtilityPath();
+	return FString::Printf(TEXT("%s --host %s --access-token \"%s\""),
+			InCommandLineArgs, *WriteToString<64>(EffectiveDomain), *GetAccessToken());
+}
 
-		FString CommandLineArgs = FString::Printf(TEXT("%s --host %s --access-token \"%s\""),
-			*InCommandLineArgs, *WriteToString<64>(EffectiveDomain), *GetAccessToken());
+bool
+FBuildServiceInstance::InvokeZenUtilitySync(FStringView InCommandLineArgs)
+{
+	FString ZenUtilityPath = UE::Zen::GetLocalInstallUtilityPath();
 
-		if (PreInvocation)
-		{
-			PreInvocation(CommandLineArgs);
-		}
+	FString CommandLineArgs(InCommandLineArgs);
 
-		FString AbsoluteUtilityPath = FPaths::ConvertRelativePathToFull(ZenUtilityPath);
-		FMonitoredProcess MonitoredUtilityProcess(AbsoluteUtilityPath, *CommandLineArgs, FPaths::GetPath(AbsoluteUtilityPath), true);
-		if (!MonitoredUtilityProcess.Launch())
-		{
-			UE_LOG(LogBuildServiceInstance, Warning, TEXT("Failed to launch zen utility to gather build data: '%s'."), *AbsoluteUtilityPath);
-			if (OnComplete)
-			{
-				OnComplete(false);
-			}
-			return;
-		}
+	FString AbsoluteUtilityPath = FPaths::ConvertRelativePathToFull(ZenUtilityPath);
+	FMonitoredProcess MonitoredUtilityProcess(AbsoluteUtilityPath, *CommandLineArgs, FPaths::GetPath(AbsoluteUtilityPath), true);
+	if (!MonitoredUtilityProcess.Launch())
+	{
+		UE_LOG(LogBuildServiceInstance, Warning, TEXT("Failed to launch zen utility to gather build data: '%s'."), *AbsoluteUtilityPath);
+		return false;
+	}
 
-		const uint64 StartTime = FPlatformTime::Cycles64();
-		while (MonitoredUtilityProcess.Update())
+	const uint64 StartTime = FPlatformTime::Cycles64();
+	while (MonitoredUtilityProcess.Update())
+	{
+		double Duration = FPlatformTime::ToSeconds64(FPlatformTime::Cycles64() - StartTime);
+		if (Duration > 120.0)
 		{
-			double Duration = FPlatformTime::ToSeconds64(FPlatformTime::Cycles64() - StartTime);
-			if (Duration > 120.0)
-			{
-				MonitoredUtilityProcess.Cancel(true);
-				UE_LOG(LogBuildServiceInstance, Warning, TEXT("Cancelled launch of zen utility for gathering build data: '%s' due to timeout."), *AbsoluteUtilityPath);
-
-				// Wait for execution to be terminated
-				while (MonitoredUtilityProcess.Update())
-				{
-					Duration = FPlatformTime::ToSeconds64(FPlatformTime::Cycles64() - StartTime);
-					if (Duration > 15.0)
-					{
-						UE_LOG(LogBuildServiceInstance, Warning, TEXT("Cancelled launch of zen utility for gathering build data: '%s'. Failed waiting for termination."), *AbsoluteUtilityPath);
-						break;
-					}
-					FPlatformProcess::Sleep(0.2f);
-				}
+			MonitoredUtilityProcess.Cancel(true);
+			UE_LOG(LogBuildServiceInstance, Warning, TEXT("Cancelled launch of zen utility for gathering build data: '%s' due to timeout."), *AbsoluteUtilityPath);
 
-				FString OutputString = MonitoredUtilityProcess.GetFullOutputWithoutDelegate();
-				UE_LOG(LogBuildServiceInstance, Warning, TEXT("Launch of zen utility for gathering build data: '%s' failed. Output: '%s'"), *AbsoluteUtilityPath, *OutputString);
-				if (OnComplete)
+			// Wait for execution to be terminated
+			while (MonitoredUtilityProcess.Update())
+			{
+				Duration = FPlatformTime::ToSeconds64(FPlatformTime::Cycles64() - StartTime);
+				if (Duration > 15.0)
 				{
-					OnComplete(false);
+					UE_LOG(LogBuildServiceInstance, Warning, TEXT("Cancelled launch of zen utility for gathering build data: '%s'. Failed waiting for termination."), *AbsoluteUtilityPath);
+					break;
 				}
-				return;
+				FPlatformProcess::Sleep(0.2f);
 			}
-			FPlatformProcess::Sleep(0.1f);
+
+			FString OutputString = MonitoredUtilityProcess.GetFullOutputWithoutDelegate();
+			UE_LOG(LogBuildServiceInstance, Warning, TEXT("Launch of zen utility for gathering build data: '%s' failed. Output: '%s'"), *AbsoluteUtilityPath, *OutputString);
+			return false;
 		}
+		FPlatformProcess::Sleep(0.1f);
+	}
+
+	FString OutputString = MonitoredUtilityProcess.GetFullOutputWithoutDelegate();
+	if (MonitoredUtilityProcess.GetReturnCode() != 0)
+	{
+		UE_LOG(LogBuildServiceInstance, Warning, TEXT("Unexpected return code after launch of zen utility for gathering build data: '%s' (%d). Output: '%s'"), *AbsoluteUtilityPath, MonitoredUtilityProcess.GetReturnCode(), *OutputString);
+		return false;
+	}
 
-		FString OutputString = MonitoredUtilityProcess.GetFullOutputWithoutDelegate();
-		if (MonitoredUtilityProcess.GetReturnCode() != 0)
+	return true;
+}
+
+void
+FBuildServiceInstance::InvokeZenUtility(FStringView InCommandLineArgs, FPreZenUtilityInvocation&& PreInvocation, FOnZenUtilityInvocationComplete&& OnComplete)
+{
+	UE::Tasks::Launch(UE_SOURCE_LOCATION,
+	[this, InCommandLineArgs = FString(InCommandLineArgs), PreInvocation = MoveTemp(PreInvocation), OnComplete = MoveTemp(OnComplete)]() mutable
+	{
+		if (PreInvocation)
 		{
-			UE_LOG(LogBuildServiceInstance, Warning, TEXT("Unexpected return code after launch of zen utility for gathering build data: '%s' (%d). Output: '%s'"), *AbsoluteUtilityPath, MonitoredUtilityProcess.GetReturnCode(), *OutputString);
-			if (OnComplete)
-			{
-				OnComplete(false);
-			}
-			return;
+			PreInvocation(InCommandLineArgs);
 		}
 
+		bool bSuccessfulInvocation = InvokeZenUtilitySync(InCommandLineArgs);
+
 		if (OnComplete)
 		{
-			OnComplete(true);
+			OnComplete(bSuccessfulInvocation);
 		}
 	});
 }
@@ -650,6 +670,13 @@ FBuildServiceInstance::BuildTransferThreadLoop()
 		}
 		TSharedRef<FBuildTransferCommand> TransferCommand(OptionalCommand.GetValue());
 
+		if (TransferCommand->bCancelRequested.load(std::memory_order_relaxed))
+		{
+			FWriteScopeLock _(TransferCommand->Lock);
+			TransferCommand->Status = EBuildTransferStatus::Canceled;
+			continue;
+		}
+
  		FString ZenUtilityPath = FPaths::ConvertRelativePathToFull(UE::Zen::GetLocalInstallUtilityPath());
  		FString OidcTokenExecutableFilename = FPaths::ConvertRelativePathToFull(FDesktopPlatformModule::TryGet()->GetOidcTokenExecutableFilename(FPaths::RootDir()));
 
@@ -663,18 +690,67 @@ FBuildServiceInstance::BuildTransferThreadLoop()
 			break;
 		case EBuildTransferType::Oplog:
 			{
-				UE::Zen::FZenLocalServiceRunContext RunContext;
-				if (UE::Zen::TryGetLocalServiceRunContext(RunContext) && !TransferCommand->ProjectFilePath.IsEmpty())
-				{
-					UE::Zen::StartLocalService(RunContext);
-				}
 				FString DestinationProjectId, DestinationOplogId;
 				FString SplitDelim(TEXT("/"));
 				TransferCommand->Destination.Split(SplitDelim, &DestinationProjectId, &DestinationOplogId);
+
+				UE::Zen::FZenLocalServiceRunContext RunContext;
+				uint16 LocalPort = 8558;
+				if (UE::Zen::TryGetLocalServiceRunContext(RunContext))
+				{
+					if (!UE::Zen::IsLocalServiceRunning(*RunContext.GetDataPath(), &LocalPort))
+					{
+						UE::Zen::StartLocalService(RunContext);
+						UE::Zen::IsLocalServiceRunning(*RunContext.GetDataPath(), &LocalPort);
+					}
+				}
+
+				if (!TransferCommand->ProjectFilePath.IsEmpty())
+				{
+					FString OplogLifetimeMarkerPath = FPaths::Combine(FPaths::GetPath(TransferCommand->ProjectFilePath),
+														TEXT("Saved"),
+														TEXT("Cooked"),
+														DestinationOplogId,
+														TEXT("ue.projectstore"));
+
+					Zen::FZenProjectStoreWriter ProjectStoreWriter(*OplogLifetimeMarkerPath);
+					ProjectStoreWriter.Write(TEXT("[::1]"), LocalPort, true, DestinationProjectId, DestinationOplogId);
+
+					FString RootDir = FPaths::RootDir();
+					FString EngineDir = FPaths::EngineDir();
+					FPaths::NormalizeDirectoryName(EngineDir);
+					FString ProjectDir = FPaths::GetPath(TransferCommand->ProjectFilePath);
+					FPaths::NormalizeDirectoryName(ProjectDir);
+					FString ProjectPath = TransferCommand->ProjectFilePath;
+					FPaths::NormalizeFilename(ProjectPath);
+
+					IPlatformFile& PlatformFile = FPlatformFileManager::Get().GetPlatformFile();
+					FString AbsServerRoot = PlatformFile.ConvertToAbsolutePathForExternalAppForRead(*RootDir);
+					FString AbsEngineDir = PlatformFile.ConvertToAbsolutePathForExternalAppForRead(*EngineDir);
+					FString AbsProjectDir = PlatformFile.ConvertToAbsolutePathForExternalAppForRead(*ProjectDir);
+					FString ProjectFilePath = PlatformFile.ConvertToAbsolutePathForExternalAppForRead(*ProjectPath);
+
+					InvokeZenUtilitySync(FString::Printf(TEXT("project-create --project \"%s\" --rootdir \"%s\" --enginedir \"%s\" --projectdir \"%s\" --projectfile \"%s\""),
+						*DestinationProjectId, *AbsServerRoot, *AbsEngineDir, *AbsProjectDir, *ProjectFilePath));
+
+					InvokeZenUtilitySync(FString::Printf(TEXT("oplog-create --project \"%s\" --oplog \"%s\" --gcpath \"%s\""),
+						*DestinationProjectId, *DestinationOplogId, *OplogLifetimeMarkerPath));
+				}
+				else
+				{
+					InvokeZenUtilitySync(FString::Printf(TEXT("project-create --project \"%s\""),
+						*DestinationProjectId));
+
+					InvokeZenUtilitySync(FString::Printf(TEXT("oplog-create --project \"%s\" --oplog \"%s\""),
+						*DestinationProjectId, *DestinationOplogId));
+				}
+
 				CommandLineArgs = FString::Printf(TEXT("oplog-import %s %s --clean --builds %s --namespace %s --bucket %s --builds-id %s --access-token \"%s\""),
 					*DestinationProjectId, *DestinationOplogId, *WriteToString<64>(EffectiveDomain),
 					*TransferCommand->Namespace, *TransferCommand->Bucket,
 					*WriteToString<64>(TransferCommand->BuildId), *GetAccessToken());
+
+
 			}
 			break;
 		}
@@ -691,6 +767,7 @@ FBuildServiceInstance::BuildTransferThreadLoop()
 			FWriteScopeLock _(TransferCommand->Lock);
 			TransferCommand->Status = EBuildTransferStatus::Active;
 		}
+
 		if (!MonitoredUtilityProcess.Launch())
 		{
 			UE_LOG(LogBuildServiceInstance, Warning, TEXT("Failed to launch zen utility to download build: '%s'."), *ZenUtilityPath);
@@ -707,7 +784,7 @@ FBuildServiceInstance::BuildTransferThreadLoop()
 				return;
 			}
 			double Duration = FPlatformTime::ToSeconds64(FPlatformTime::Cycles64() - StartTime);
-			if (TransferCommand->bCancelRequested.load(std::memory_order_relaxed))
+			if (bBuildTransferThreadStopping.load(std::memory_order_relaxed) || TransferCommand->bCancelRequested.load(std::memory_order_relaxed))
 			{
 				MonitoredUtilityProcess.Cancel(true);
 				UE_LOG(LogBuildServiceInstance, Display, TEXT("Canceled launch of zen utility for downloading build data: '%s'."), *ZenUtilityPath);
@@ -731,24 +808,26 @@ FBuildServiceInstance::BuildTransferThreadLoop()
 				UE_LOG(LogBuildServiceInstance, Display, TEXT("Launch of zen utility for downloading build data: '%s' canceled."), *ZenUtilityPath);
 				FWriteScopeLock _(TransferCommand->Lock);
 				TransferCommand->Status = EBuildTransferStatus::Canceled;
-				continue;
+				break;
 			}
 			FPlatformProcess::Sleep(0.1f);
 		}
 
-		FString OutputString = MonitoredUtilityProcess.GetFullOutputWithoutDelegate();
-		if (MonitoredUtilityProcess.GetReturnCode() == 0)
+		if (!TransferCommand->bCancelRequested.load(std::memory_order_relaxed))
 		{
-			FWriteScopeLock _(TransferCommand->Lock);
-			TransferCommand->ReturnCode = MonitoredUtilityProcess.GetReturnCode();
-			TransferCommand->Status = EBuildTransferStatus::Succeeded;
-		}
-		else
-		{
-			UE_LOG(LogBuildServiceInstance, Warning, TEXT("Unexpected return code after launch of zen utility for downloading build data: '%s' (%d). Output: '%s'"), *ZenUtilityPath, MonitoredUtilityProcess.GetReturnCode(), *TransferCommand->Output);
-			FWriteScopeLock _(TransferCommand->Lock);
-			TransferCommand->ReturnCode = MonitoredUtilityProcess.GetReturnCode();
-			TransferCommand->Status = EBuildTransferStatus::Failed;
+			if (MonitoredUtilityProcess.GetReturnCode() == 0)
+			{
+				FWriteScopeLock _(TransferCommand->Lock);
+				TransferCommand->ReturnCode = MonitoredUtilityProcess.GetReturnCode();
+				TransferCommand->Status = EBuildTransferStatus::Succeeded;
+			}
+			else
+			{
+				UE_LOG(LogBuildServiceInstance, Warning, TEXT("Unexpected return code after launch of zen utility for downloading build data: '%s' (%d). Output: '%s'"), *ZenUtilityPath, MonitoredUtilityProcess.GetReturnCode(), *TransferCommand->Output);
+				FWriteScopeLock _(TransferCommand->Lock);
+				TransferCommand->ReturnCode = MonitoredUtilityProcess.GetReturnCode();
+				TransferCommand->Status = EBuildTransferStatus::Failed;
+			}
 		}
 	}
 }
diff --git a/Engine/Source/Developer/Zen/Private/ZenProjectStoreWriter.cpp b/Engine/Source/Developer/Zen/Private/ZenProjectStoreWriter.cpp
new file mode 100644
--- /dev/null
+++ b/Engine/Source/Developer/Zen/Private/ZenProjectStoreWriter.cpp
@@ -0,0 +1,99 @@
+// Copyright Epic Games, Inc. All Rights Reserved.
+
+#include "Experimental/ZenProjectStoreWriter.h"
+
+#if UE_WITH_ZEN
+
+#include "Containers/StringView.h"
+#include "Experimental/ZenServerInterface.h"
+#include "HAL/FileManager.h"
+#include "Policies/PrettyJsonPrintPolicy.h"
+#include "Serialization/Archive.h"
+#include "Serialization/JsonSerializer.h"
+#include "Serialization/JsonWriter.h"
+#include "SocketSubsystem.h"
+
+namespace UE::Zen
+{
+
+FZenProjectStoreWriter::FZenProjectStoreWriter(const TCHAR* TargetFileName)
+: OplogMarker(IFileManager::Get().CreateFileWriter(TargetFileName))
+{
+}
+
+bool FZenProjectStoreWriter::Write(FStringView ServiceHostName, uint16 ServicePort, bool bIsRunningLocally, FStringView ProjectId, FStringView OplogId)
+{
+	if (!OplogMarker)
+	{
+		return false;
+	}
+
+	TSharedRef<TJsonWriter<UTF8CHAR, TPrettyJsonPrintPolicy<UTF8CHAR>>> Writer = TJsonWriterFactory<UTF8CHAR, TPrettyJsonPrintPolicy<UTF8CHAR>>::Create(OplogMarker.Get());
+	Writer->WriteObjectStart();
+	Writer->WriteObjectStart(TEXT("zenserver"));
+	Writer->WriteValue(TEXT("islocalhost"), bIsRunningLocally);
+	Writer->WriteValue(TEXT("hostname"), ServiceHostName);
+	if (bIsRunningLocally)
+	{
+		ISocketSubsystem* SocketSubsystem = ISocketSubsystem::Get();
+		if (SocketSubsystem != nullptr)
+		{
+			TArray<TSharedPtr<FInternetAddr>> Addresses;
+			if (SocketSubsystem->GetLocalAdapterAddresses(Addresses))
+			{
+				Writer->WriteArrayStart("remotehostnames");
+				for (const TSharedPtr<FInternetAddr>& Address : Addresses)
+				{
+					Writer->WriteValue(Address->ToString(false));
+				}
+				FString MachineHostName;
+				if (SocketSubsystem->GetHostName(MachineHostName))
+				{
+					// Try to acquire FQDN
+					FAddressInfoResult GAIRequest = SocketSubsystem->GetAddressInfo(*MachineHostName, nullptr, EAddressInfoFlags::AllResultsWithMapping | EAddressInfoFlags::OnlyUsableAddresses | EAddressInfoFlags::AllResults | EAddressInfoFlags::FQDomainName, NAME_None);
+					if (GAIRequest.ReturnCode == SE_NO_ERROR && !GAIRequest.CanonicalNameResult.IsEmpty())
+					{
+						MachineHostName = GAIRequest.CanonicalNameResult;
+					}
+
+					Writer->WriteValue(TEXT("hostname://") + MachineHostName);
+				}
+#if PLATFORM_MAC
+				// Store the Bonjour hostname when on a mac so that the client can attempt to connect via USB.
+				// This is done since the "Link Local" IP of the USB connection changes everytime
+				// the cable is plugged in, so there is no stable IP to save.
+				bool bAppleTarget = TargetPlatformFName == TEXT("IOS") || TargetPlatformFName == TEXT("TVOS");
+				if (bAppleTarget)
+				{
+					FString MacBonjourName = FPlatformMisc::GetBonjourName();
+					
+					// If there is no Bonjour name, fallback to the mac's hostname
+					if (MacBonjourName.IsEmpty())
+					{
+						MacBonjourName = FPlatformProcess::ComputerName();
+					}
+					
+					if (!MacBonjourName.IsEmpty())
+					{
+						Writer->WriteValue(TEXT("macserver://") + MacBonjourName);
+					}
+				}
+#endif
+				Writer->WriteArrayEnd();
+			}
+		}
+	}
+	Writer->WriteValue(TEXT("hostport"), ServicePort);
+	Writer->WriteValue(TEXT("projectid"), ProjectId);
+	Writer->WriteValue(TEXT("oplogid"), OplogId);
+	Writer->WriteObjectEnd();
+	Writer->WriteObjectEnd();
+	Writer->Close();
+
+	OplogMarker.Reset();
+	return true;
+}
+
+}
+
+#endif // UE_WITH_ZEN
diff --git a/Engine/Source/Developer/Zen/Public/Experimental/BuildServerInterface.h b/Engine/Source/Developer/Zen/Public/Experimental/BuildServerInterface.h
--- a/Engine/Source/Developer/Zen/Public/Experimental/BuildServerInterface.h
+++ b/Engine/Source/Developer/Zen/Public/Experimental/BuildServerInterface.h
@@ -96,6 +96,8 @@ public:
 	public:
 		UE_API FString GetDescription() const;
 		UE_API EBuildTransferStatus GetStatus() const;
+
+		UE_API void RequestCancel();
 	private:
 		TSharedPtr<FBuildTransferCommand> BuildTransferCommand;
 		friend class UE::Zen::Build::FBuildServiceInstance;
@@ -132,6 +134,8 @@ private:
 	bool AcquireAccessToken();
 	void SetAccessTokenAndUnlock(FScopeLock &Lock, FStringView Token, double RefreshDelay = 0.0);
 	[[nodiscard]] FString GetAccessToken() const;
+	FString AddZenUtilityBuildServerArguments(const TCHAR* InCommandLineArgs) const;
+	bool InvokeZenUtilitySync(FStringView CommandlineArgs);
 	void InvokeZenUtility(FStringView CommandlineArgs, FPreZenUtilityInvocation&& PreInvocation, FOnZenUtilityInvocationComplete&& OnComplete);
 	void CallOnRefreshNamespacesAndBucketsComplete(FOnRefreshNamespacesAndBucketsComplete&& InOnRefreshNamespacesAndBucketsComplete);
 
diff --git a/Engine/Source/Developer/Zen/Public/Experimental/ZenProjectStoreWriter.h b/Engine/Source/Developer/Zen/Public/Experimental/ZenProjectStoreWriter.h
new file mode 100644
--- /dev/null
+++ b/Engine/Source/Developer/Zen/Public/Experimental/ZenProjectStoreWriter.h
@@ -0,0 +1,32 @@
+// Copyright Epic Games, Inc. All Rights Reserved.
+
+#pragma once
+
+#include "Containers/StringFwd.h"
+#include "Experimental/ZenGlobals.h"
+#include "Serialization/Archive.h"
+#include "Templates/UniquePtr.h"
+
+#define UE_API ZEN_API
+
+#if UE_WITH_ZEN
+
+namespace UE::Zen
+{
+
+class FZenProjectStoreWriter
+{
+public:
+	UE_API FZenProjectStoreWriter(const TCHAR* TargetFileName);
+	UE_API ~FZenProjectStoreWriter() = default;
+
+	UE_API bool Write(FStringView ServiceHostName, uint16 ServicePort, bool bIsRunningLocally, FStringView ProjectId, FStringView OplogId);
+private:
+	TUniquePtr<FArchive> OplogMarker;
+};
+
+} // namespace UE
+
+#endif // UE_WITH_ZEN
+
+#undef UE_API
diff --git a/Engine/Source/Editor/StorageServerWidgets/Internal/SBuildActivity.h b/Engine/Source/Editor/StorageServerWidgets/Internal/SBuildActivity.h
--- a/Engine/Source/Editor/StorageServerWidgets/Internal/SBuildActivity.h
+++ b/Engine/Source/Editor/StorageServerWidgets/Internal/SBuildActivity.h
@@ -47,7 +47,13 @@ class STORAGESERVERWIDGETS_API SBuildActivity : public SCompoundWidget
 private:
 	TSharedRef<ITableRow> GenerateBuildActivityRow(TSharedPtr<FBuildActivity> InItem, const TSharedRef<STableViewBase>& InOwningTable);
 	TSharedRef<SWidget> GetGridPanel();
+	TSharedPtr<SWidget> OnGetBuildActivityContextMenuContent();
 
+	bool CanCancelBuildActivity() const;
+	void CancelBuildActivity() const;
+	void ClearCompleted();
+
+	TSharedPtr<SListView<TSharedPtr<SBuildActivity::FBuildActivity, ESPMode::ThreadSafe>>> BuildActivityListView;
 	TArray<TSharedPtr<SBuildActivity::FBuildActivity, ESPMode::ThreadSafe>> BuildActivities;
 
 	SVerticalBox::FSlot* GridSlot = nullptr;
diff --git a/Engine/Source/Editor/StorageServerWidgets/Internal/SBuildSelection.h b/Engine/Source/Editor/StorageServerWidgets/Internal/SBuildSelection.h
--- a/Engine/Source/Editor/StorageServerWidgets/Internal/SBuildSelection.h
+++ b/Engine/Source/Editor/StorageServerWidgets/Internal/SBuildSelection.h
@@ -45,6 +45,7 @@ class STORAGESERVERWIDGETS_API SBuildSelection : public SCompoundWidget
 		FString Suffix;
 		FDateTime CreatedAt;
 		FString Category;
+		FString Job;
 		TMap<FString, UE::Zen::Build::FBuildServiceInstance::FBuildRecord> PerPlatformBuilds;
 	};
 
@@ -58,22 +59,41 @@ private:
 		PackagedBuild,
 		EditorPreCompiledBinary,
 		EditorInstalledBuild,
-		Unknown
+		Unknown,
+		Count
 	};
 	EBuildType GetSelectedBuildType() const;
 	void RebuildLists();
-	void ReselectDestination(TSharedPtr<FBuildGroup> Item);
 	void RegenerateBuildGroups(UE::BuildSelection::Internal::FListBuildsState& ListBuildsState);
 	void RegenerateActivePlatformFilters();
 	void ValidateBuildGroupSelection();
 	TSharedRef<SWidget> OnGenerateTextBlockFromString(TSharedPtr<FString> Item);
 	bool BuildGroupIsSelectableOrNavigable(TSharedPtr<FBuildGroup> InItem) const;
 	TSharedRef<ITableRow> GenerateBuildGroupRow(TSharedPtr<FBuildGroup> InItem, const TSharedRef<STableViewBase>& InOwningTable);
+	TSharedPtr<SWidget> OnGetBuildGroupContextMenuContent() const;
 	void OnOpenDestinationDirectoryClicked();
 	TSharedRef<SWidget> GetBuildDestinationPanel();
 	TSharedRef<SWidget> GetGridPanel();
 	void BuildGroupSelectionChanged(TSharedPtr<FBuildGroup> Item, ESelectInfo::Type SelectInfo);
 
+	FString& GetUserSelectedDestination();
+	FString GetUserSelectedDestination() const;
+	FString GetDefaultDestination() const;
+	FString GetEffectiveDestination() const;
+
+	static FString SanitizeForPath(const FString& InString);
+	static FString SanitizeForZenId(const FString& InString);
+	static FString SanitizeBucketSegment(const FString& InString);
+
+	FReply ExploreDestination_OnClicked();
+
+	FString BuildUrl;
+	FString BuildId;
+	FString BuildVersion;
+	FString BranchName;
+	FString BucketFormattedBranchName;
+	int EffectiveCompatibleChangelist = 0;
+
 	TSharedPtr<SComboBox<TSharedPtr<FString>>> ProjectWidget;
 	TSharedPtr<SComboBox<TSharedPtr<FString>>> StreamWidget;
 	TSharedPtr<SComboBox<TSharedPtr<FString>>> BuildTypeWidget;
@@ -90,14 +110,14 @@ private:
 	TArray<TSharedPtr<FString>> PlatformList;
 
 	TArray<FString> ActivePlatformFilters;
-	FString DestinationFolderPath;
-	FString DestinationZenProjectId;
+	FString UserSelectedDestinations[(int)EBuildType::Count];
+	bool bAppendBuildNameToDestinations[(int)EBuildType::Count];
 
 	TSharedPtr<FString> SelectedStream;
 	TSharedPtr<FString> SelectedProject;
 	TSharedPtr<FString> SelectedBuildType;
 
-	TSharedPtr<SListView<TSharedPtr<SBuildSelection::FBuildGroup, ESPMode::ThreadSafe>>> BuildListView;
+	TSharedPtr<SListView<TSharedPtr<SBuildSelection::FBuildGroup, ESPMode::ThreadSafe>>> BuildGroupListView;
 
 	std::atomic<uint32> BuildListRefreshesInProgress = 0;
 	std::atomic<uint32> BuildRefreshGeneration = 0;
diff --git a/Engine/Source/Editor/StorageServerWidgets/Private/SBuildActivity.cpp b/Engine/Source/Editor/StorageServerWidgets/Private/SBuildActivity.cpp
--- a/Engine/Source/Editor/StorageServerWidgets/Private/SBuildActivity.cpp
+++ b/Engine/Source/Editor/StorageServerWidgets/Private/SBuildActivity.cpp
@@ -1,6 +1,7 @@
 // Copyright Epic Games, Inc. All Rights Reserved.
 
 #include "SBuildActivity.h"
+#include "Framework/MultiBox/MultiBoxBuilder.h"
 #include "Internationalization/FastDecimalFormat.h"
 #include "Math/BasicMathExpressionEvaluator.h"
 #include "Math/UnitConversion.h"
@@ -81,9 +82,10 @@ TSharedRef<SWidget> SBuildActivity::GetGridPanel()
 	.AutoHeight()
 	.Padding(FMargin(ColumnMargin, RowMargin))
 	[
-		SNew(SListView<FBuildActivityPtr>)
+		SAssignNew(BuildActivityListView, SListView<FBuildActivityPtr>)
 			.ListItemsSource(&BuildActivities)
 			.OnGenerateRow(this, &SBuildActivity::GenerateBuildActivityRow)
+			.OnContextMenuOpening(this, &SBuildActivity::OnGetBuildActivityContextMenuContent)
 			.HeaderRow
 			(
 				SNew(SHeaderRow)
@@ -103,6 +105,101 @@ TSharedRef<SWidget> SBuildActivity::GetGridPanel()
 	return Panel;
 }
 
+TSharedPtr<SWidget> SBuildActivity::OnGetBuildActivityContextMenuContent()
+{
+	TArray<FBuildActivityPtr> SelectedItems = BuildActivityListView->GetSelectedItems();
+
+	if (SelectedItems.IsEmpty())
+	{
+		return nullptr;
+	}
+
+	const bool bCloseAfterSelection = true;
+	const bool bCloseSelfOnly = false;
+	const bool bSearchable = false;
+	const bool bRecursivelySearchable = false;
+	FMenuBuilder MenuBuilder(bCloseAfterSelection,
+		nullptr,
+		TSharedPtr<FExtender>(),
+		bCloseSelfOnly,
+		&FCoreStyle::Get(),
+		bSearchable,
+		NAME_None,
+		bRecursivelySearchable);
+
+	MenuBuilder.AddMenuEntry(
+		LOCTEXT("BuildActivity_Cancel", "Cancel"),
+		LOCTEXT("BuildActivity_CancelTooltip", "Cancel the pending activity"),
+		FSlateIcon(),
+		FUIAction(
+			FExecuteAction::CreateRaw(this, &SBuildActivity::CancelBuildActivity),
+			FCanExecuteAction::CreateRaw(this, &SBuildActivity::CanCancelBuildActivity)
+		),
+		NAME_None,
+		EUserInterfaceActionType::Button
+	);
+
+	MenuBuilder.AddMenuEntry(
+		LOCTEXT("BuildActivity_ClearCompleted", "Clear completed"),
+		LOCTEXT("BuildActivity_ClearCompletedTooltip", "Clear all completed items"),
+		FSlateIcon(),
+		FUIAction(
+			FExecuteAction::CreateRaw(this, &SBuildActivity::ClearCompleted)
+		),
+		NAME_None,
+		EUserInterfaceActionType::Button
+	);
+
+	return MenuBuilder.MakeWidget();
+}
+
+bool SBuildActivity::CanCancelBuildActivity() const
+{
+	using namespace UE::Zen::Build;
+
+	TArray<FBuildActivityPtr> SelectedItems = BuildActivityListView->GetSelectedItems();
+	if (SelectedItems.Num() != 1)
+	{
+		return false;
+	}
+	FBuildServiceInstance::EBuildTransferStatus Status = SelectedItems[0]->Transfer.GetStatus();
+	return (Status == FBuildServiceInstance::EBuildTransferStatus::Queued) || (Status == FBuildServiceInstance::EBuildTransferStatus::Active);
+}
+
+void SBuildActivity::CancelBuildActivity() const
+{
+	using namespace UE::Zen::Build;
+
+	TArray<FBuildActivityPtr> SelectedItems = BuildActivityListView->GetSelectedItems();
+	if (SelectedItems.Num() != 1)
+	{
+		return;
+	}
+	SelectedItems[0]->Transfer.RequestCancel();
+}
+
+void SBuildActivity::ClearCompleted()
+{
+	using namespace UE::Zen::Build;
+
+	TArray<FBuildActivityPtr> ActivitiesToRemove;
+	for (FBuildActivityPtr BuildActivity : BuildActivities)
+	{
+		FBuildServiceInstance::EBuildTransferStatus Status = BuildActivity->Transfer.GetStatus();
+		if ((Status == FBuildServiceInstance::EBuildTransferStatus::Failed) ||
+			(Status == FBuildServiceInstance::EBuildTransferStatus::Canceled) ||
+			(Status == FBuildServiceInstance::EBuildTransferStatus::Succeeded))
+		{
+			ActivitiesToRemove.AddUnique(BuildActivity);
+		}
+	}
+
+	for (FBuildActivityPtr BuildActivity : ActivitiesToRemove)
+	{
+		BuildActivities.Remove(BuildActivity);
+	}
+}
+
 void
 SBuildActivityTableRow::Construct(const FArguments& InArgs,
 	const TSharedRef<STableViewBase>& InOwnerTableView,
diff --git a/Engine/Source/Editor/StorageServerWidgets/Private/SBuildSelection.cpp b/Engine/Source/Editor/StorageServerWidgets/Private/SBuildSelection.cpp
--- a/Engine/Source/Editor/StorageServerWidgets/Private/SBuildSelection.cpp
+++ b/Engine/Source/Editor/StorageServerWidgets/Private/SBuildSelection.cpp
@@ -3,6 +3,9 @@
 #include "SBuildSelection.h"
 
 #include "DesktopPlatformModule.h"
+#include "Framework/MultiBox/MultiBoxBuilder.h"
+#include "HAL/PlatformApplicationMisc.h"
+#include "HAL/PlatformFileManager.h"
 #include "Internationalization/FastDecimalFormat.h"
 #include "Math/BasicMathExpressionEvaluator.h"
 #include "Math/UnitConversion.h"
@@ -10,6 +13,7 @@
 #include "Misc/ExpressionParser.h"
 #include "Misc/Paths.h"
 #include "Misc/UProjectInfo.h"
+#include "Modules/BuildVersion.h"
 #include "SMultiSelectComboBox.h"
 #include "Styling/StyleColors.h"
 #include "Widgets/Images/SImage.h"
@@ -65,6 +69,23 @@ void SBuildSelection::Construct(const FArguments& InArgs)
 	BuildServiceInstance = InArgs._BuildServiceInstance;
 	OnBuildTransferStarted = InArgs._OnBuildTransferStarted;
 
+	FBuildVersion CurrentBuildVersion;
+	if (FBuildVersion::TryRead(FBuildVersion::GetDefaultFileName(), CurrentBuildVersion))
+	{
+		BuildUrl = CurrentBuildVersion.BuildUrl;
+		BuildId = CurrentBuildVersion.BuildId;
+		BuildVersion = CurrentBuildVersion.BuildVersion;
+		BranchName = CurrentBuildVersion.BranchName;
+		BucketFormattedBranchName = SanitizeBucketSegment(BranchName);
+
+		EffectiveCompatibleChangelist = CurrentBuildVersion.GetEffectiveCompatibleChangelist();
+	}
+
+	for (bool& bAppendBuildNameToDestination : bAppendBuildNameToDestinations)
+	{
+		bAppendBuildNameToDestination = false;
+	}
+
 	if (TSharedPtr<UE::Zen::Build::FBuildServiceInstance> ServiceInstance = BuildServiceInstance.Get())
 	{
 		ServiceInstance->OnRefreshNamespacesAndBucketsComplete().AddSP(this, &SBuildSelection::RebuildLists);
@@ -136,49 +157,49 @@ void SBuildSelection::RebuildLists()
 		TMultiMap<FString, FString> NamespacesAndBuckets = ServiceInstance->GetNamespacesAndBuckets();
 
 		auto StringToSegmentViews = [](const FString& Str, TArray<FStringView>& OutViews)
-		{
-			FStringView WorkingStringView(Str);
-			int32 CurrentIndex = 0;
-			while (WorkingStringView.FindChar(TCHAR('.'), CurrentIndex))
 			{
-				if (CurrentIndex != 0)
+				FStringView WorkingStringView(Str);
+				int32 CurrentIndex = 0;
+				while (WorkingStringView.FindChar(TCHAR('.'), CurrentIndex))
 				{
-					OutViews.Add(WorkingStringView.Left(CurrentIndex));
+					if (CurrentIndex != 0)
+					{
+						OutViews.Add(WorkingStringView.Left(CurrentIndex));
+					}
+					WorkingStringView.RightChopInline(CurrentIndex + 1);
 				}
-				WorkingStringView.RightChopInline(CurrentIndex+1);
-			}
-			if (!WorkingStringView.IsEmpty())
-			{
-				OutViews.Add(WorkingStringView);
-			}
-		};
+				if (!WorkingStringView.IsEmpty())
+				{
+					OutViews.Add(WorkingStringView);
+				}
+			};
 
 		NamespacesAndBuckets.GetKeys(Namespaces);
 
-		auto ConvertToSharedPtrs = [] (TArray<FString>& Strings, TArray<TSharedPtr<FString>>& SharedStrings)
-		{
-			Strings.StableSort();
-			for (FString& String : Strings)
+		auto ConvertToSharedPtrs = [](TArray<FString>& Strings, TArray<TSharedPtr<FString>>& SharedStrings)
 			{
-				SharedStrings.Add(MakeShared<FString>(MoveTemp(String)));
-			}
-		};
+				Strings.StableSort();
+				for (FString& String : Strings)
+				{
+					SharedStrings.Add(MakeShared<FString>(MoveTemp(String)));
+				}
+			};
 
 		auto ConformSelection = [](TSharedPtr<FString>& SelectedItem, const TArray<TSharedPtr<FString>>& SelectionList)
-		{
-			if (!SelectedItem)
 			{
-				SelectedItem = SelectionList.IsEmpty() ? nullptr : SelectionList[0];
-				return;
-			}
+				if (!SelectedItem)
+				{
+					SelectedItem = SelectionList.IsEmpty() ? nullptr : SelectionList[0];
+					return;
+				}
 
-			const TSharedPtr<FString>* FoundSelectionListItem = SelectionList.FindByPredicate([&SelectedItem](const TSharedPtr<FString>& Item)
-			{
-				return *Item == *SelectedItem;
-			});
+				const TSharedPtr<FString>* FoundSelectionListItem = SelectionList.FindByPredicate([&SelectedItem](const TSharedPtr<FString>& Item)
+					{
+						return *Item == *SelectedItem;
+					});
 
-			SelectedItem = FoundSelectionListItem ? *FoundSelectionListItem : SelectionList.IsEmpty() ? nullptr : SelectionList[0];
-		};
+				SelectedItem = FoundSelectionListItem ? *FoundSelectionListItem : SelectionList.IsEmpty() ? nullptr : SelectionList[0];
+			};
 
 		const uint32 SegmentIndexProject = 0;
 		const uint32 SegmentIndexBuildType = 1;
@@ -187,6 +208,7 @@ void SBuildSelection::RebuildLists()
 		const uint32 SegmentIndexNum = 4;
 
 		// Stream list generation and selection conforming
+		int32 DefaultStreamIndex = 0;
 		TMultiMap<FStringView, TArray<FStringView>> NamespacesToBucketSegmentViews;
 		for (const TPair<FString, FString>& NamespaceAndBucket : NamespacesAndBuckets)
 		{
@@ -196,10 +218,24 @@ void SBuildSelection::RebuildLists()
 
 			if (BucketSegmentViews.Num() == SegmentIndexNum)
 			{
-				Streams.AddUnique(FString(BucketSegmentViews[SegmentIndexStream]));
+				int32 CurrentIndex = Streams.AddUnique(FString(BucketSegmentViews[SegmentIndexStream]));
+				if (!BranchName.IsEmpty() && (Streams[CurrentIndex] == BucketFormattedBranchName))
+				{
+					DefaultStreamIndex = CurrentIndex;
+				}
 			}
 		}
+		FString DefaultStreamValue = Streams[DefaultStreamIndex];
 		ConvertToSharedPtrs(Streams, StreamList);
+		// Must search to find DefaultStreamIndex because it will have re-sorted in ConverToSharedPtrs
+		const TSharedPtr<FString>* DefaultStreamListItem = StreamList.FindByPredicate([&DefaultStreamValue](const TSharedPtr<FString>& Item)
+			{
+				return *Item == DefaultStreamValue;
+			});
+		if (!SelectedStream && DefaultStreamListItem)
+		{
+			SelectedStream = *DefaultStreamListItem;
+		}
 		ConformSelection(SelectedStream, StreamList);
 
 		// Project list generation and selection conforming
@@ -325,7 +361,7 @@ void SBuildSelection::RebuildLists()
 							ExecuteOnGameThread(UE_SOURCE_LOCATION,
 								[this]
 								{
-									BuildListView->RequestListRefresh();
+									BuildGroupListView->RequestListRefresh();
 									BuildListRefreshesInProgress.fetch_sub(1);
 								});
 						}
@@ -354,42 +390,6 @@ void SBuildSelection::RebuildLists()
 	}
 }
 
-void SBuildSelection::ReselectDestination(TSharedPtr<FBuildGroup> Item)
-{
-	// TODO: Select the destination info (path & zen project id based on the selected item)
-	if (!SelectedBuildType)
-	{
-		return;
-	}
-
-	EBuildType BuildType = GetSelectedBuildType();
-	FString ProjectFilename = FUProjectDictionary::GetDefault().GetProjectPathForGame(**SelectedProject);
-	if (ProjectFilename.IsEmpty())
-	{
-		if (BuildType != EBuildType::Oplog)
-		{
-			DestinationFolderPath = FPaths::Combine(FPaths::ProjectSavedDir(), TEXT("DownloadedBuilds"));
-		}
-		return;
-	}
-
-	switch (BuildType)
-	{
-	case EBuildType::Oplog:
-		DestinationZenProjectId = FApp::GetZenStoreProjectIdForProject(ProjectFilename);
-		break;
-	case EBuildType::StagedBuild:
-		DestinationFolderPath = FPaths::Combine(FPaths::GetPath(ProjectFilename), TEXT("Saved"), TEXT("StagedBuilds"));
-		break;
-	case EBuildType::PackagedBuild:
-		DestinationFolderPath = FPaths::Combine(FPaths::GetPath(ProjectFilename), TEXT("Saved"), TEXT("Packages"));
-		break;
-	default:
-		DestinationFolderPath = FPaths::Combine(FPaths::GetPath(ProjectFilename), TEXT("Saved"), TEXT("DownloadedBuilds"));
-		break;
-	}
-}
-
 void SBuildSelection::RegenerateBuildGroups(UE::BuildSelection::Internal::FListBuildsState& ListBuildsState)
 {
 	using namespace UE::Zen::Build;
@@ -455,6 +455,12 @@ void SBuildSelection::RegenerateBuildGroups(UE::BuildSelection::Internal::FListB
 
 			if (BuildGroup->DisplayName.IsEmpty())
 			{
+				FString Job;
+				if (FCbFieldView JobField = BuildRecord.Metadata["job"]; JobField.HasValue() && !JobField.HasError())
+				{
+					Job = FUTF8ToTCHAR(JobField.AsString());
+				}
+
 				FString Category;
 				if (FCbFieldView TemplateIdField = BuildRecord.Metadata["hordeTemplateId"]; TemplateIdField.HasValue() && !TemplateIdField.HasError())
 				{
@@ -522,6 +528,7 @@ void SBuildSelection::RegenerateBuildGroups(UE::BuildSelection::Internal::FListB
 				BuildGroup->CommitIdentifier = CommitIdentifier;
 				BuildGroup->Category = Category;
 				BuildGroup->CreatedAt = CreatedAt;
+				BuildGroup->Job = Job;
 			}
 
 			BuildGroup->PerPlatformBuilds.FindOrAdd(BuildState.Platform, MoveTemp(BuildRecord));
@@ -553,8 +560,8 @@ void SBuildSelection::RegenerateActivePlatformFilters()
 
 void SBuildSelection::ValidateBuildGroupSelection()
 {
-	BuildListView->UpdateSelectionSet();
-	TArray<FBuildSelectionBuildGroupPtr> SelectedItems = BuildListView->GetSelectedItems();
+	BuildGroupListView->UpdateSelectionSet();
+	TArray<FBuildSelectionBuildGroupPtr> SelectedItems = BuildGroupListView->GetSelectedItems();
 	if (SelectedItems.IsEmpty())
 	{
 		return;
@@ -563,7 +570,7 @@ void SBuildSelection::ValidateBuildGroupSelection()
 	{
 		if (!BuildGroupIsSelectableOrNavigable(SelectedItem))
 		{
-			BuildListView->SetItemSelection(SelectedItem, false);
+			BuildGroupListView->SetItemSelection(SelectedItem, false);
 		}
 	}
 }
@@ -607,6 +614,52 @@ TSharedRef<ITableRow> SBuildSelection::GenerateBuildGroupRow(FBuildSelectionBuil
 		});
 }
 
+TSharedPtr<SWidget> SBuildSelection::OnGetBuildGroupContextMenuContent() const
+{
+	TArray<FBuildSelectionBuildGroupPtr> SelectedItems = BuildGroupListView->GetSelectedItems();
+
+	if (SelectedItems.IsEmpty())
+	{
+		return nullptr;
+	}
+
+	const bool bCloseAfterSelection = true;
+	const bool bCloseSelfOnly = false;
+	const bool bSearchable = false;
+	const bool bRecursivelySearchable = false;
+	FMenuBuilder MenuBuilder(bCloseAfterSelection,
+		nullptr,
+		TSharedPtr<FExtender>(),
+		bCloseSelfOnly,
+		&FCoreStyle::Get(),
+		bSearchable,
+		NAME_None,
+		bRecursivelySearchable);
+
+	MenuBuilder.AddMenuEntry(
+		LOCTEXT("BuildSelection_CopyCommit", "Copy commit"),
+		LOCTEXT("BuildSelection_CopyCommit", "Copies the commit or changelist number"),
+		FSlateIcon(),
+		FUIAction(
+			FExecuteAction::CreateLambda([this]
+				{
+					TArray<FBuildSelectionBuildGroupPtr> SelectedItems = BuildGroupListView->GetSelectedItems();
+					if (SelectedItems.Num() == 1)
+					{
+						FPlatformApplicationMisc::ClipboardCopy(*SelectedItems[0]->CommitIdentifier);
+					}
+				})
+		),
+		NAME_None,
+		EUserInterfaceActionType::Button
+	);
+
+
+	MenuBuilder.EndSection();
+
+	return MenuBuilder.MakeWidget();
+}
+
 void SBuildSelection::BuildGroupSelectionChanged(FBuildSelectionBuildGroupPtr Item, ESelectInfo::Type SelectInfo)
 {
 	using namespace UE::Zen::Build;
@@ -623,14 +676,12 @@ void SBuildSelection::BuildGroupSelectionChanged(FBuildSelectionBuildGroupPtr It
 		return;
 	}
 
-	ReselectDestination(Item);
-
 	int32 Row = 0;
 	int32 Column = 0;
 
 	for (const TSharedPtr<FString>& Platform : PlatformList)
 	{
-		TSharedPtr<SCheckBox> AssociatedCheckbox;
+		TSharedPtr<SCheckBox> CurrentCheckbox;
 		SelectedGroupPlatformGrid->AddSlot(Column, Row)
 			.Padding(0,0,0,2)
 			[
@@ -643,7 +694,7 @@ void SBuildSelection::BuildGroupSelectionChanged(FBuildSelectionBuildGroupPtr It
 				.AutoWidth()
 				.VAlign(VAlign_Center)
 				[
-					SAssignNew(AssociatedCheckbox, SCheckBox)
+					SAssignNew(CurrentCheckbox, SCheckBox)
 					.IsChecked_Lambda([this, Item, Platform]
 					{
 						bool bIsChecked = Item->PerPlatformBuilds.Contains(*Platform) && SelectedGroupSelectedPlatforms.Contains(*Platform);
@@ -669,14 +720,14 @@ void SBuildSelection::BuildGroupSelectionChanged(FBuildSelectionBuildGroupPtr It
 					.ButtonStyle(FAppStyle::Get(), "InvisibleButton")
 					.IsFocusable(false)
 
-					.OnClicked_Lambda([this, AssociatedCheckbox, Platform]()
+					.OnClicked_Lambda([this, CurrentCheckbox, Platform]()
 						{
 							ECheckBoxState NewState = ECheckBoxState::Checked;
-							if (AssociatedCheckbox->IsChecked())
+							if (CurrentCheckbox->IsChecked())
 							{
 								NewState = ECheckBoxState::Unchecked;
 							}
-							AssociatedCheckbox->SetIsChecked(NewState);
+							CurrentCheckbox->SetIsChecked(NewState);
 							if (NewState == ECheckBoxState::Checked)
 							{
 								SelectedGroupSelectedPlatforms.AddUnique(*Platform);
@@ -693,6 +744,43 @@ void SBuildSelection::BuildGroupSelectionChanged(FBuildSelectionBuildGroupPtr It
 						.Text(FText::FromString(*Platform))
 					]
 				]
+				+SHorizontalBox::Slot()
+				.AutoWidth()
+				.VAlign(VAlign_Center)
+				[
+					SNew(SButton)
+					.ButtonStyle(&FAppStyle::Get().GetWidgetStyle<FButtonStyle>("SimpleButton"))
+					.ToolTipText(LOCTEXT("BuildSelection_BrowseToBuildUrl", "Browse to build url"))
+					.IsEnabled_Lambda([this, Item, Platform]
+					{
+						if (FBuildServiceInstance::FBuildRecord* BuildRecord = Item->PerPlatformBuilds.Find(*Platform))
+						{
+							if (FCbFieldView BuildUrlField = BuildRecord->Metadata["buildurl"]; BuildUrlField.HasValue() && !BuildUrlField.HasError())
+							{
+								return true;
+							}
+						}
+						return false;
+					})
+					.OnClicked_Lambda([this, Item, Platform]
+					{
+						if (FBuildServiceInstance::FBuildRecord* BuildRecord = Item->PerPlatformBuilds.Find(*Platform))
+						{
+							if (FCbFieldView BuildUrlField = BuildRecord->Metadata["buildurl"]; BuildUrlField.HasValue() && !BuildUrlField.HasError())
+							{
+								FString Url(FUTF8ToTCHAR(BuildUrlField.AsString()));
+								FPlatformProcess::LaunchURL(*Url, nullptr, nullptr);
+								return FReply::Handled();
+							}
+						}
+						return FReply::Unhandled();
+					})
+					[
+						SNew(SImage)
+						.Image(FAppStyle::Get().GetBrush("Zen.BrowserView"))
+						.ColorAndOpacity(FSlateColor::UseForeground())
+					]
+				]
 			];
 
 		if (++Row > 2)
@@ -703,6 +791,146 @@ void SBuildSelection::BuildGroupSelectionChanged(FBuildSelectionBuildGroupPtr It
 	}
 }
 
+FString& SBuildSelection::GetUserSelectedDestination()
+{
+	EBuildType BuildType = GetSelectedBuildType();
+	int BuildTypeIndex = (int)BuildType;
+	return UserSelectedDestinations[BuildTypeIndex];
+}
+
+FString SBuildSelection::GetUserSelectedDestination() const
+{
+	EBuildType BuildType = GetSelectedBuildType();
+	int BuildTypeIndex = (int)BuildType;
+	return BuildType == EBuildType::Oplog ? SanitizeForZenId(UserSelectedDestinations[BuildTypeIndex]) : SanitizeForPath(UserSelectedDestinations[BuildTypeIndex]);
+}
+
+FString SBuildSelection::GetDefaultDestination() const
+{
+	EBuildType BuildType = GetSelectedBuildType();
+	IPlatformFile& PlatformFile = FPlatformFileManager::Get().GetPlatformFile();
+
+	FString ProjectFilename = FUProjectDictionary::GetDefault().GetProjectPathForGame(**SelectedProject);
+	if (ProjectFilename.IsEmpty())
+	{
+		if (BuildType == EBuildType::Oplog)
+		{
+			return *SelectedProject;
+		}
+
+		return PlatformFile.ConvertToAbsolutePathForExternalAppForRead(
+				*FPaths::FPaths::Combine(FPaths::EngineSavedDir(), TEXT("DownloadedBuilds"))
+			);
+	}
+
+	switch (BuildType)
+	{
+	case EBuildType::StagedBuild:
+		return PlatformFile.ConvertToAbsolutePathForExternalAppForRead(
+				*FPaths::Combine(FPaths::GetPath(ProjectFilename), TEXT("Saved"), TEXT("StagedBuilds"))
+			);
+	case EBuildType::PackagedBuild:
+		return PlatformFile.ConvertToAbsolutePathForExternalAppForRead(
+				*FPaths::Combine(FPaths::GetPath(ProjectFilename), TEXT("Saved"), TEXT("Packages"))
+			);
+	case EBuildType::Oplog:
+		return FApp::GetZenStoreProjectIdForProject(ProjectFilename);
+	}
+	return PlatformFile.ConvertToAbsolutePathForExternalAppForRead(
+				*FPaths::Combine(FPaths::GetPath(ProjectFilename), TEXT("Saved"), TEXT("DownloadedBuilds"))
+			);
+}
+
+FString SBuildSelection::GetEffectiveDestination() const
+{
+	FString UserSelectedDestination = GetUserSelectedDestination();
+	return UserSelectedDestination.IsEmpty() ? GetDefaultDestination() : UserSelectedDestination;
+}
+
+FString SBuildSelection::SanitizeForPath(const FString& InString)
+{
+	FString OutString = InString;
+	OutString = FPaths::MakeValidFileName(InString);
+	return OutString;
+}
+
+FString SBuildSelection::SanitizeForZenId(const FString& InString)
+{
+	FString OutString = InString;
+	for (int32 Index = 0; Index < OutString.Len(); ++Index)
+	{
+		if (!FChar::IsIdentifier(OutString[Index]))
+		{
+			OutString[Index] = TCHAR('_');
+		}
+	}
+	return OutString;
+}
+
+FString SBuildSelection::SanitizeBucketSegment(const FString& InString)
+{
+	TStringBuilder<64> OutputBuilder;
+	for (int32 CharIndex = 0; CharIndex < InString.Len(); ++CharIndex)
+	{
+		TCHAR Character = FChar::ToLower(InString[CharIndex]);
+
+		if (Character == TCHAR('.'))
+		{
+			if (OutputBuilder.Len() > 0)
+			{
+				if (OutputBuilder.LastChar() == TCHAR('-'))
+				{
+					OutputBuilder.RemoveSuffix(1);
+				}
+				OutputBuilder.AppendChar(Character);
+			}
+		}
+		else if (FChar::IsIdentifier(Character))
+		{
+			OutputBuilder.AppendChar(Character);
+		}
+		else if (OutputBuilder.Len() > 0 &&
+			OutputBuilder.LastChar() != TCHAR('-'))
+		{
+			OutputBuilder.AppendChar(TCHAR('-'));
+		}
+	}
+
+	FString OutString = OutputBuilder.ToString();
+
+	// Trim leading and trailing dashes
+	bool bCharsWereRemoved;
+	do
+	{
+		OutString.TrimCharInline(TEXT('-'), &bCharsWereRemoved);
+	} while (bCharsWereRemoved);
+
+	return OutString;
+}
+
+FReply SBuildSelection::ExploreDestination_OnClicked()
+{
+	if (GetSelectedBuildType() == EBuildType::Oplog)
+	{
+		UE::Zen::FZenLocalServiceRunContext RunContext;
+		uint16 LocalPort = 8558;
+		if (UE::Zen::TryGetLocalServiceRunContext(RunContext))
+		{
+			if (!UE::Zen::IsLocalServiceRunning(*RunContext.GetDataPath(), &LocalPort))
+			{
+				UE::Zen::StartLocalService(RunContext);
+				UE::Zen::IsLocalServiceRunning(*RunContext.GetDataPath(), &LocalPort);
+			}
+		}
+		FPlatformProcess::LaunchURL(*FString::Printf(TEXT("http://localhost:%d/dashboard/?page=project&project=%s"), LocalPort, *GetEffectiveDestination()), nullptr, nullptr);
+	}
+	else
+	{
+		FPlatformProcess::ExploreFolder(*GetEffectiveDestination());
+	}
+	return FReply::Handled();
+}
+
 void SBuildSelection::OnOpenDestinationDirectoryClicked()
 {
 	bool bDirectorySelected = false;
@@ -714,13 +942,15 @@ void SBuildSelection::OnOpenDestinationDirectoryClicked()
 		bDirectorySelected = DesktopPlatform->OpenDirectoryDialog(
 			FSlateApplication::Get().FindBestParentWindowHandleForDialogs(nullptr),
 			Title,
-			DestinationFolderPath,
-			DestinationFolderPath);
+			GetEffectiveDestination(),
+			GetUserSelectedDestination());
 	}
 }
 
 TSharedRef<SWidget> SBuildSelection::GetBuildDestinationPanel()
 {
+	TSharedPtr<SCheckBox> AppendBuildNameCheckbox;
+
 	return
 	SNew(SVerticalBox)
 
@@ -731,7 +961,7 @@ TSharedRef<SWidget> SBuildSelection::GetBuildDestinationPanel()
 		SNew(SHorizontalBox)
 		.Visibility_Lambda([this]()
 		{
-			TArray<FBuildSelectionBuildGroupPtr> SelectedItems = BuildListView->GetSelectedItems();
+			TArray<FBuildSelectionBuildGroupPtr> SelectedItems = BuildGroupListView->GetSelectedItems();
 			return (SelectedItems.Num() == 0) || GetSelectedBuildType() == EBuildType::Oplog ? EVisibility::Collapsed : EVisibility::Visible;
 		})
 		+ SHorizontalBox::Slot()
@@ -741,17 +971,17 @@ TSharedRef<SWidget> SBuildSelection::GetBuildDestinationPanel()
 			SNew(SEditableTextBox)
 			.OverflowPolicy(ETextOverflowPolicy::MiddleEllipsis)
 			.MinDesiredWidth(200.0f)
-			.Text_Lambda([this]()
+			.HintText_Lambda([this]()
 			{
-				return FText::FromString(DestinationFolderPath);
+				return FText::FromString(GetDefaultDestination());
 			})
-			.OnTextChanged_Lambda([this](const FText& Text)
+			.Text_Lambda([this]()
 			{
-				DestinationFolderPath = Text.ToString();
+				return FText::FromString(GetUserSelectedDestination());
 			})
 			.OnTextCommitted_Lambda([this](const FText& Text, const ETextCommit::Type CommitType)
 			{
-				DestinationFolderPath = Text.ToString();
+				GetUserSelectedDestination() = Text.ToString();
 			})
 		]
 		+ SHorizontalBox::Slot()
@@ -780,7 +1010,7 @@ TSharedRef<SWidget> SBuildSelection::GetBuildDestinationPanel()
 		SNew(SHorizontalBox)
 		.Visibility_Lambda([this]()
 		{
-			TArray<FBuildSelectionBuildGroupPtr> SelectedItems = BuildListView->GetSelectedItems();
+			TArray<FBuildSelectionBuildGroupPtr> SelectedItems = BuildGroupListView->GetSelectedItems();
 			return (SelectedItems.Num() == 0) || GetSelectedBuildType() != EBuildType::Oplog ? EVisibility::Collapsed : EVisibility::Visible;
 		})
 		+ SHorizontalBox::Slot()
@@ -790,19 +1020,84 @@ TSharedRef<SWidget> SBuildSelection::GetBuildDestinationPanel()
 			SNew(SEditableTextBox)
 			.OverflowPolicy(ETextOverflowPolicy::MiddleEllipsis)
 			.MinDesiredWidth(200.0f)
-			.Text_Lambda([this]()
+			.HintText_Lambda([this]()
 			{
-				return FText::FromString(DestinationZenProjectId);
+				return FText::FromString(GetDefaultDestination());
 			})
-			.OnTextChanged_Lambda([this](const FText& Text)
+			.Text_Lambda([this]()
 			{
-				DestinationZenProjectId = Text.ToString();
+				return FText::FromString(GetUserSelectedDestination());
 			})
 			.OnTextCommitted_Lambda([this](const FText& Text, const ETextCommit::Type CommitType)
 			{
-				DestinationZenProjectId = Text.ToString();
+				GetUserSelectedDestination() = Text.ToString();
 			})
 		]
+	]
+
+	+ SVerticalBox::Slot()
+	.AutoHeight()
+	.VAlign(VAlign_Top)
+	[
+		SNew(SHorizontalBox)
+		+SHorizontalBox::Slot()
+		.AutoWidth()
+		.VAlign(VAlign_Center)
+		[
+			SAssignNew(AppendBuildNameCheckbox, SCheckBox)
+			.IsChecked_Lambda([this]
+			{
+				EBuildType BuildType = GetSelectedBuildType();
+				int BuildTypeIndex = (int)BuildType;
+				return bAppendBuildNameToDestinations[BuildTypeIndex] ? ECheckBoxState::Checked : ECheckBoxState::Unchecked;
+			})
+			.OnCheckStateChanged_Lambda([this](ECheckBoxState InNewState)
+			{
+				EBuildType BuildType = GetSelectedBuildType();
+				int BuildTypeIndex = (int)BuildType;
+				if (InNewState == ECheckBoxState::Checked)
+				{
+					bAppendBuildNameToDestinations[BuildTypeIndex] = true;
+				}
+				else if (InNewState == ECheckBoxState::Unchecked)
+				{
+					bAppendBuildNameToDestinations[BuildTypeIndex] = false;
+				}
+			})
+		]
+		+SHorizontalBox::Slot()
+		.AutoWidth()
+		.VAlign(VAlign_Center)
+		[
+			SNew(SButton)
+			.ButtonStyle(FAppStyle::Get(), "InvisibleButton")
+			.IsFocusable(false)
+			.OnClicked_Lambda([this, AppendBuildNameCheckbox]()
+				{
+					EBuildType BuildType = GetSelectedBuildType();
+					int BuildTypeIndex = (int)BuildType;
+					ECheckBoxState NewState = ECheckBoxState::Checked;
+					if (AppendBuildNameCheckbox->IsChecked())
+					{
+						NewState = ECheckBoxState::Unchecked;
+					}
+					AppendBuildNameCheckbox->SetIsChecked(NewState);
+					if (NewState == ECheckBoxState::Checked)
+					{
+						bAppendBuildNameToDestinations[BuildTypeIndex] = true;
+					}
+					else if (NewState == ECheckBoxState::Unchecked)
+					{
+						bAppendBuildNameToDestinations[BuildTypeIndex] = false;
+					}
+					return FReply::Handled();
+				})
+			[
+				SNew(STextBlock)
+				.Justification(ETextJustify::Left)
+				.Text(LOCTEXT("BuildSelection_DestinationAppendBuildName", "Append build name"))
+			]
+		]
 	];
 }
 
@@ -1038,10 +1333,11 @@ TSharedRef<SWidget> SBuildSelection::GetGridPanel()
 		.HAlign(HAlign_Fill)
 		.VAlign(VAlign_Fill)
 		[
-			SAssignNew(BuildListView, SListView<FBuildSelectionBuildGroupPtr>)
+			SAssignNew(BuildGroupListView, SListView<FBuildSelectionBuildGroupPtr>)
 				.ListItemsSource(&BuildGroups)
 				.OnGenerateRow(this, &SBuildSelection::GenerateBuildGroupRow)
 				.OnSelectionChanged(this, &SBuildSelection::BuildGroupSelectionChanged)
+				.OnContextMenuOpening(this, &SBuildSelection::OnGetBuildGroupContextMenuContent)
 				.SelectionMode(ESelectionMode::Single)
 				.OnIsSelectableOrNavigable(this, &SBuildSelection::BuildGroupIsSelectableOrNavigable)
 				.IsEnabled_Lambda([this]
@@ -1128,7 +1424,7 @@ TSharedRef<SWidget> SBuildSelection::GetGridPanel()
 		SNew(SHorizontalBox)
 		.Visibility_Lambda([this]()
 		{
-			return !!BuildListRefreshesInProgress || BuildListView->GetNumItemsSelected() == 0 ? EVisibility::Collapsed : EVisibility::Visible;
+			return !!BuildListRefreshesInProgress || BuildGroupListView->GetNumItemsSelected() == 0 ? EVisibility::Collapsed : EVisibility::Visible;
 		})
 		+SHorizontalBox::Slot()
 		.HAlign(HAlign_Left)
@@ -1163,10 +1459,34 @@ TSharedRef<SWidget> SBuildSelection::GetGridPanel()
 			.VAlign(VAlign_Top)
 			.AutoHeight()
 			[
-				SNew(STextBlock)
-					.ColorAndOpacity(TitleColor)
-					.Font(TitleFont)
-					.Text(LOCTEXT("BuildSelection_Destination", "Destination"))
+				SNew(SHorizontalBox)
+				+SHorizontalBox::Slot()
+				.HAlign(HAlign_Left)
+				.VAlign(VAlign_Bottom)
+				.AutoWidth()
+				[
+					SNew(STextBlock)
+						.ColorAndOpacity(TitleColor)
+						.Font(TitleFont)
+						.Text(LOCTEXT("BuildSelection_Destination", "Destination"))
+				]
+				+SHorizontalBox::Slot()
+				.HAlign(HAlign_Left)
+				.AutoWidth()
+				[
+					SNew(SButton)
+					.ButtonStyle(&FAppStyle::Get().GetWidgetStyle<FButtonStyle>("SimpleButton"))
+					.ToolTipText(LOCTEXT("BuildSelection_DestinationExploreTooltip", "Explore the destination"))
+					.OnClicked(this, &SBuildSelection::ExploreDestination_OnClicked)
+					[
+						SNew(SImage)
+						.Image_Lambda([this]
+						{
+							return GetSelectedBuildType() == EBuildType::Oplog ? FAppStyle::Get().GetBrush("Zen.BrowserView") : FAppStyle::Get().GetBrush("Zen.FolderView");
+						})
+						.ColorAndOpacity(FSlateColor::UseForeground())
+					]
+				]
 			]
 			+SVerticalBox::Slot()
 			.HAlign(HAlign_Fill)
@@ -1185,13 +1505,13 @@ TSharedRef<SWidget> SBuildSelection::GetGridPanel()
 				.ButtonStyle(FAppStyle::Get(), "Button")
 				.IsEnabled_Lambda([this]()
 				{
-					const bool bDestinationValid = GetSelectedBuildType() == EBuildType::Oplog ? !DestinationZenProjectId.IsEmpty() : !DestinationFolderPath.IsEmpty();
+					const bool bDestinationValid = !GetEffectiveDestination().IsEmpty();
 					if (!bDestinationValid)
 					{
 						return false;
 					}
 
-					TArray<FBuildSelectionBuildGroupPtr> SelectedItems = BuildListView->GetSelectedItems();
+					TArray<FBuildSelectionBuildGroupPtr> SelectedItems = BuildGroupListView->GetSelectedItems();
 					if (SelectedItems.Num() == 0)
 					{
 						return false;
@@ -1212,7 +1532,7 @@ TSharedRef<SWidget> SBuildSelection::GetGridPanel()
 
 					if (TSharedPtr<FBuildServiceInstance> ServiceInstance = BuildServiceInstance.Get())
 					{
-						TArray<FBuildSelectionBuildGroupPtr> SelectedItems = BuildListView->GetSelectedItems();
+						TArray<FBuildSelectionBuildGroupPtr> SelectedItems = BuildGroupListView->GetSelectedItems();
 						if (SelectedItems.Num() == 0)
 						{
 							return FReply::Handled();
@@ -1223,20 +1543,41 @@ TSharedRef<SWidget> SBuildSelection::GetGridPanel()
 							if (FBuildServiceInstance::FBuildRecord* BuildRecord = SelectedItems[0]->PerPlatformBuilds.Find(PlatformForSelectedGroup))
 							{
 								FString Bucket = FString::Printf(TEXT("%s.%s.%s.%s"), **SelectedProject, **SelectedBuildType, **SelectedStream, *PlatformForSelectedGroup);
-								if (GetSelectedBuildType() == EBuildType::Oplog)
+
+								FString DestinationPlatformName = PlatformForSelectedGroup;
+								if (FCbFieldView CookPlatformField = BuildRecord->Metadata["cookPlatform"]; CookPlatformField.HasValue() && !CookPlatformField.HasError())
+								{
+									DestinationPlatformName = *WriteToString<64>(CookPlatformField.AsString());
+								}
+
+								EBuildType BuildType = GetSelectedBuildType();
+								int BuildTypeIndex = (int)BuildType;
+								const bool bAppendBuildNameToDestination = bAppendBuildNameToDestinations[BuildTypeIndex];
+
+								if (BuildType == EBuildType::Oplog)
 								{
-									if (FCbFieldView CookPlatformField = BuildRecord->Metadata["cookPlatform"]; CookPlatformField.HasValue() && !CookPlatformField.HasError())
+									FString ProjectFilePath = FUProjectDictionary::GetDefault().GetProjectPathForGame(**SelectedProject);
+									FString DestinationProjectId = GetEffectiveDestination();
+									FString DestinationOplogId = DestinationPlatformName;
+									if (bAppendBuildNameToDestination)
 									{
-										FString ProjectFilePath = FUProjectDictionary::GetDefault().GetProjectPathForGame(**SelectedProject);
-										FString DestinationOplogId = *WriteToString<64>(CookPlatformField.AsString());
-										FBuildServiceInstance::FBuildTransfer BuildTransfer =
-											ServiceInstance->StartOplogBuildTransfer(BuildRecord->BuildId, DestinationZenProjectId, DestinationOplogId, ProjectFilePath, SelectedItems[0]->Namespace, Bucket);
-										OnBuildTransferStarted.ExecuteIfBound(BuildTransfer, SelectedItems[0]->DisplayName, PlatformForSelectedGroup);
+										DestinationOplogId = FString::Printf(TEXT("%s.%s"), *DestinationOplogId, *SanitizeForZenId(SelectedItems[0]->DisplayName));
 									}
+									FBuildServiceInstance::FBuildTransfer BuildTransfer =
+										ServiceInstance->StartOplogBuildTransfer(BuildRecord->BuildId, DestinationProjectId, DestinationOplogId, ProjectFilePath, SelectedItems[0]->Namespace, Bucket);
+									OnBuildTransferStarted.ExecuteIfBound(BuildTransfer, SelectedItems[0]->DisplayName, PlatformForSelectedGroup);
 								}
 								else
 								{
-									FString DestinationFolder = FPaths::Combine(DestinationFolderPath, PlatformForSelectedGroup);
+									FString DestinationFolder;
+									if (bAppendBuildNameToDestination)
+									{
+										DestinationFolder = FPaths::Combine(GetEffectiveDestination(), SanitizeForPath(SelectedItems[0]->DisplayName), DestinationPlatformName);
+									}
+									else
+									{
+										DestinationFolder = FPaths::Combine(GetEffectiveDestination(), DestinationPlatformName);
+									}
 									FBuildServiceInstance::FBuildTransfer BuildTransfer =
 										ServiceInstance->StartBuildTransfer(BuildRecord->BuildId, DestinationFolder, SelectedItems[0]->Namespace, Bucket);
 									OnBuildTransferStarted.ExecuteIfBound(BuildTransfer, SelectedItems[0]->DisplayName, PlatformForSelectedGroup);
diff --git a/Engine/Source/Programs/BuildStorageTool/Private/BuildStorageTool.cpp b/Engine/Source/Programs/BuildStorageTool/Private/BuildStorageTool.cpp
--- a/Engine/Source/Programs/BuildStorageTool/Private/BuildStorageTool.cpp
+++ b/Engine/Source/Programs/BuildStorageTool/Private/BuildStorageTool.cpp
@@ -187,9 +187,9 @@ public:
 
 					// Selection panel
 					+ SVerticalBox::Slot()
-					.Padding(0.0f, 10.0f, 5.0f, 0.0f)
+					.Padding(0.0f, 10.0f, 5.0f, 10.0f)
 					.HAlign(HAlign_Fill)
-					.VAlign(VAlign_Top)
+					.VAlign(VAlign_Fill)
 					[
 						SNew(SBuildSelection)
 						.ZenServiceInstance(ZenServiceInstanceManager.ToSharedRef(), &UE::Zen::FServiceInstanceManager::GetZenServiceInstance)
diff --git a/Engine/Source/Programs/BuildStorageTool/Private/BuildStorageToolStyle.cpp b/Engine/Source/Programs/BuildStorageTool/Private/BuildStorageToolStyle.cpp
--- a/Engine/Source/Programs/BuildStorageTool/Private/BuildStorageToolStyle.cpp
+++ b/Engine/Source/Programs/BuildStorageTool/Private/BuildStorageToolStyle.cpp
@@ -16,6 +16,7 @@
 #define EDITOR_BOX_BRUSH(RelativePath, ...) BOX_BRUSH("../Editor/Slate/" RelativePath, __VA_ARGS__)
 #define EDITOR_BORDER_BRUSH(RelativePath, ...) BORDER_BRUSH("../Editor/Slate/" RelativePath, __VA_ARGS__)
 #define RootToContentDir ContentFromEngine
+#define RootToCoreContentDir ContentFromEngine
 
 TSharedPtr< FSlateStyleSet > FBuildStorageToolStyle::StyleSet = nullptr;
 
@@ -63,6 +64,9 @@ TSharedRef< FSlateStyleSet > FBuildStorageToolStyle::Create()
 	Style.Set("Zen.BrowseContent", new EDITOR_IMAGE_BRUSH_SVG("Starship/Common/BrowseContent", Icon16x16));
 	Style.Set("Zen.Clipboard", new EDITOR_IMAGE_BRUSH_SVG("Starship/Common/Clipboard", Icon16x16));
 
+	Style.Set("Zen.FolderView", new CORE_IMAGE_BRUSH_SVG("Starship/Common/folder-open", Icon12x12));
+	Style.Set("Zen.BrowserView", new EDITOR_IMAGE_BRUSH_SVG("Starship/Common/WebBrowser", Icon12x12));
+
 	const FTextBlockStyle DefaultText = FTextBlockStyle()
 		.SetFont(DEFAULT_FONT("Bold", 10))
 		.SetColorAndOpacity(FSlateColor::UseForeground())
