Fixed text conflict from LOCTEXT macro used with duplicated key (in MassEntityDebugger and in StorageServerWidgets).

#rnx

diff --git a/Engine/Source/Editor/MassEntityDebugger/Private/SMassBreakpointsView.cpp b/Engine/Source/Editor/MassEntityDebugger/Private/SMassBreakpointsView.cpp
--- a/Engine/Source/Editor/MassEntityDebugger/Private/SMassBreakpointsView.cpp
+++ b/Engine/Source/Editor/MassEntityDebugger/Private/SMassBreakpointsView.cpp
@@ -181,7 +181,7 @@ TSharedRef<ITableRow> SMassBreakpointsView::OnGenerateBreakpointRow(TSharedPtr<F
 					? LOCTEXT("SelectedEntity", "Selected Entity")
 					: FText::FromString(InItem->Entity.DebugGetDescription()))
 				.ToolTip(FSlateApplicationBase::Get().MakeToolTip(InItem->bBreakForSelectedEntity
-					? LOCTEXT("SelectedEntity", "Selected Entity (Gameplay Debugger)")
+					? LOCTEXT("SelectedEntityGameplayDebugger", "Selected Entity (Gameplay Debugger)")
 					: FText::FromString(InItem->Entity.DebugGetDescription())))
             ]
 #else
diff --git a/Engine/Source/Editor/StorageServerWidgets/Private/SBuildSelection.cpp b/Engine/Source/Editor/StorageServerWidgets/Private/SBuildSelection.cpp
--- a/Engine/Source/Editor/StorageServerWidgets/Private/SBuildSelection.cpp
+++ b/Engine/Source/Editor/StorageServerWidgets/Private/SBuildSelection.cpp
@@ -638,7 +638,7 @@ TSharedPtr<SWidget> SBuildSelection::OnGetBuildGroupContextMenuContent() const
 
 	MenuBuilder.AddMenuEntry(
 		LOCTEXT("BuildSelection_CopyCommit", "Copy commit"),
-		LOCTEXT("BuildSelection_CopyCommit", "Copies the commit or changelist number"),
+		LOCTEXT("BuildSelection_CopyCommit_ToolTip", "Copies the commit or changelist number"),
 		FSlateIcon(),
 		FUIAction(
 			FExecuteAction::CreateLambda([this]
@@ -1403,14 +1403,7 @@ TSharedRef<SWidget> SBuildSelection::GetGridPanel()
 						}
 					}
 				}
-				if (VisibleItemCount == 1)
-				{
-					return LOCTEXT("BuildSelection_ResultDescriptionMultiple", "1 item");
-				}
-				else
-				{
-					return FText::Format(LOCTEXT("BuildSelection_ResultDescriptionMultiple", "{0} items"), FText::AsNumber(VisibleItemCount));
-				}
+				return FText::Format(LOCTEXT("BuildSelection_ResultDescription", "{0} {0}|plural(one=item,other=items)"), FText::AsNumber(VisibleItemCount));
 			})
 	];
 
