Removed sanitation of user-supplied destination paths as it was inappropriately sanitizing as a file name instead of as a path name.  This would cause all path separators to be removed.

diff --git a/Engine/Source/Editor/StorageServerWidgets/Private/SBuildSelection.cpp b/Engine/Source/Editor/StorageServerWidgets/Private/SBuildSelection.cpp
--- a/Engine/Source/Editor/StorageServerWidgets/Private/SBuildSelection.cpp
+++ b/Engine/Source/Editor/StorageServerWidgets/Private/SBuildSelection.cpp
@@ -849,9 +849,8 @@ FString SBuildSelection::GetEffectiveDestination() const
 
 FString SBuildSelection::SanitizeForPath(const FString& InString)
 {
-	FString OutString = InString;
-	OutString = FPaths::MakeValidFileName(InString);
-	return OutString;
+	// TODO: Had to remove path sanitization as the engine only provides validation
+	return InString;
 }
 
 FString SBuildSelection::SanitizeForZenId(const FString& InString)
