Add commit filter for build selection within the BuildStorageTool

diff --git a/Engine/Source/Editor/StorageServerWidgets/Internal/SBuildSelection.h b/Engine/Source/Editor/StorageServerWidgets/Internal/SBuildSelection.h
--- a/Engine/Source/Editor/StorageServerWidgets/Internal/SBuildSelection.h
+++ b/Engine/Source/Editor/StorageServerWidgets/Internal/SBuildSelection.h
@@ -122,6 +122,7 @@ private:
 	TSharedPtr<FString> SelectedStream;
 	TSharedPtr<FString> SelectedProject;
 	TSharedPtr<FString> SelectedBuildType;
+	FString SelectedCommitFilter;
 
 	TSharedPtr<SListView<TSharedPtr<SBuildSelection::FBuildGroup, ESPMode::ThreadSafe>>> BuildGroupListView;
 	FName BuildGroupSortByColumn;
diff --git a/Engine/Source/Editor/StorageServerWidgets/Private/SBuildSelection.cpp b/Engine/Source/Editor/StorageServerWidgets/Private/SBuildSelection.cpp
--- a/Engine/Source/Editor/StorageServerWidgets/Private/SBuildSelection.cpp
+++ b/Engine/Source/Editor/StorageServerWidgets/Private/SBuildSelection.cpp
@@ -679,6 +679,11 @@ bool SBuildSelection::BuildGroupIsSelectableOrNavigable(FBuildSelectionBuildGrou
 			return false;
 		}
 	}
+
+	if (!SelectedCommitFilter.IsEmpty() && !InItem->CommitIdentifier.Contains(SelectedCommitFilter))
+	{
+		return false;
+	}
 	return true;
 }
 
@@ -694,6 +699,11 @@ TSharedRef<ITableRow> SBuildSelection::GenerateBuildGroupRow(FBuildSelectionBuil
 					return EVisibility::Collapsed;
 				}
 			}
+
+			if (!SelectedCommitFilter.IsEmpty() && !InItem->CommitIdentifier.Contains(SelectedCommitFilter))
+			{
+				return EVisibility::Collapsed;
+			}
 			return EVisibility::Visible;
 		});
 }
@@ -1380,6 +1390,53 @@ TSharedRef<SWidget> SBuildSelection::GetGridPanel()
 			]
 		]
 
+		// Commit
+		+SHorizontalBox::Slot()
+		.AutoWidth()
+		.FillWidth(1.0f)
+		.HAlign(HAlign_Fill)
+		.VAlign(VAlign_Top)
+		[
+			SNew(SVerticalBox)
+
+			+ SVerticalBox::Slot()
+			.AutoHeight()
+			.Padding(FMargin(ColumnMargin, RowMargin))
+			.VAlign(VAlign_Top)
+			[
+				SNew(STextBlock)
+				.ColorAndOpacity(TitleColor)
+				.Font(TitleFont)
+				.Text(LOCTEXT("BuildSelection_Commit", "Commit"))
+			]
+
+			+ SVerticalBox::Slot()
+			.AutoHeight()
+			.Padding(FMargin(ColumnMargin, RowMargin))
+			.HAlign(HAlign_Fill)
+			.VAlign(VAlign_Top)
+			[
+				SNew(SEditableTextBox)
+				.OverflowPolicy(ETextOverflowPolicy::MiddleEllipsis)
+				.HintText_Lambda([this]()
+				{
+					return LOCTEXT("BuildSelection_AnyCommit", "Any");
+				})
+				.Text_Lambda([this]()
+				{
+					return FText::FromString(SelectedCommitFilter);
+				})
+				.OnTextChanged_Lambda([this](const FText& Text)
+				{
+					SelectedCommitFilter = Text.ToString();
+				})
+				.OnTextCommitted_Lambda([this](const FText& Text, const ETextCommit::Type CommitType)
+				{
+					SelectedCommitFilter = Text.ToString();
+				})
+			]
+		]
+
 		// Required Platforms
 		+SHorizontalBox::Slot()
 		.AutoWidth()
