Build Storage Tool - Search for matching bucket prefixes rather then exact match of bucket names, allows for more flexibility in bucket names.

diff --git a/Engine/Source/Developer/Zen/Private/BuildServerInterface.cpp b/Engine/Source/Developer/Zen/Private/BuildServerInterface.cpp
--- a/Engine/Source/Developer/Zen/Private/BuildServerInterface.cpp
+++ b/Engine/Source/Developer/Zen/Private/BuildServerInterface.cpp
@@ -338,7 +338,7 @@ FBuildServiceInstance::RefreshNamespacesAndBuckets(FOnRefreshNamespacesAndBucket
 }
 
 void
-FBuildServiceInstance::ListBuilds(FStringView Namespace, FStringView Bucket, FOnListBuildsComplete&& InOnListBuildsComplete)
+FBuildServiceInstance::ListBuilds(FStringView Namespace, FStringView BucketPrefix, FOnListBuildsComplete&& InOnListBuildsComplete)
 {
 	if (ConnectionState.load(std::memory_order_relaxed) != EConnectionState::ConnectionSucceeded)
 	{
@@ -355,15 +355,19 @@ FBuildServiceInstance::ListBuilds(FStringView Namespace, FStringView Bucket, FOn
 	FString OutputFilePath = FPaths::CreateTempFilename(FPlatformProcess::UserTempDir(), TEXT("BuildList"), TEXT(".cbo"));
 	FPaths::MakePlatformFilename(OutputFilePath);
 
-	FString CommandLineArgs = FString::Printf(TEXT("builds list --namespace \"%s\" --bucket \"%s\" \"%s\" \"%s\""),
-		*WriteToString<64>(Namespace), *WriteToString<64>(Bucket), *QueryFilePath, *OutputFilePath);
+	FString CommandLineArgs = FString::Printf(TEXT("builds list --namespace \"%s\" \"%s\" \"%s\""),
+		*WriteToString<64>(Namespace), *QueryFilePath, *OutputFilePath);
 	CommandLineArgs = AddZenUtilityBuildServerArguments(*CommandLineArgs);
 
-	auto PreInvoke = [QueryFilePath](FString& CommandLineArgs)
+	FString BucketRegex = FString::Printf(TEXT("%s.*"), *WriteToString<64>(BucketPrefix));
+
+	auto PreInvoke = [QueryFilePath, BucketRegex](FString& CommandLineArgs)
 	{
 		TUniquePtr<FArchive> QueryFileArchive(IFileManager::Get().CreateFileWriter(*QueryFilePath, FILEWRITE_NoFail));
+
 		FCbWriter Writer;
 		Writer.BeginObject("query");
+		Writer << "bucketRegex" << *BucketRegex;
 		Writer.EndObject();
 		Writer.Save(*QueryFileArchive);
 	};
diff --git a/Engine/Source/Developer/Zen/Public/Experimental/BuildServerInterface.h b/Engine/Source/Developer/Zen/Public/Experimental/BuildServerInterface.h
--- a/Engine/Source/Developer/Zen/Public/Experimental/BuildServerInterface.h
+++ b/Engine/Source/Developer/Zen/Public/Experimental/BuildServerInterface.h
@@ -105,7 +105,7 @@ public:
 
 	inline const TMultiMap<FString, FString>& GetNamespacesAndBuckets() const { return NamespacesAndBuckets; }
 
-	UE_API void ListBuilds(FStringView Namespace, FStringView Bucket, FOnListBuildsComplete&& InOnListBuildsComplete);
+	UE_API void ListBuilds(FStringView Namespace, FStringView BucketPrefix, FOnListBuildsComplete&& InOnListBuildsComplete);
 
 	UE_API FBuildTransfer StartBuildTransfer(const FCbObjectId& BuildId, FStringView Name, FStringView DestinationFolder, FStringView Namespace, FStringView Bucket);
 	UE_API FBuildTransfer StartOplogBuildTransfer(const FCbObjectId& BuildId, FStringView Name, FStringView DestinationProjectId, FStringView DestinationOplogId, FStringView ProjectFilePath, FStringView Namespace, FStringView Bucket, FName TargetPlatformName);
diff --git a/Engine/Source/Editor/StorageServerWidgets/Private/SBuildSelection.cpp b/Engine/Source/Editor/StorageServerWidgets/Private/SBuildSelection.cpp
--- a/Engine/Source/Editor/StorageServerWidgets/Private/SBuildSelection.cpp
+++ b/Engine/Source/Editor/StorageServerWidgets/Private/SBuildSelection.cpp
@@ -45,14 +45,12 @@ namespace UE::BuildSelection::Internal
 	struct FNamespacePlatformBucketTuple
 	{
 		FString Namespace;
-		FString Platform;
-		FString Bucket;
+		FString BucketPrefix;
 	};
 
 	struct FBuildState
 	{
 		FString Namespace;
-		FString Platform;
 		TArray<UE::Zen::Build::FBuildServiceInstance::FBuildRecord> Results;
 	};
 
@@ -323,36 +321,40 @@ void SBuildSelection::RebuildLists()
 
 	if (TSharedPtr<FBuildServiceInstance> ServiceInstance = BuildServiceInstance.Get())
 	{
-		TArray<FNamespacePlatformBucketTuple> NamespacePlatformBucketTuples;
-		for (TSharedPtr<FString> Platform : PlatformList)
+		TArray<FNamespacePlatformBucketTuple> NamespaceBucketTuples;
+		FString BucketPrefix = FString::Printf(TEXT("%s.%s.%s"), **SelectedProject, **SelectedBuildType, **SelectedStream);
+		TArray<FString> NamespacesForBucket;
+
+		for (auto& Pair : BucketsToNamespaces)
 		{
-			FString Bucket = FString::Printf(TEXT("%s.%s.%s.%s"), **SelectedProject, **SelectedBuildType, **SelectedStream, **Platform);
-			TArray<FString> NamespacesForBucket;
-			BucketsToNamespaces.MultiFind(Bucket, NamespacesForBucket);
-			for (FString& Namespace : NamespacesForBucket)
+			if (Pair.Key.StartsWith(BucketPrefix))
 			{
-				NamespacePlatformBucketTuples.Emplace(MoveTemp(Namespace), *Platform, Bucket);
+				NamespacesForBucket.Add(Pair.Value);
 			}
 		}
 
+		for (FString& Namespace : NamespacesForBucket)
+		{
+			NamespaceBucketTuples.Emplace(MoveTemp(Namespace), BucketPrefix);
+		}
+
 		TSharedPtr<FListBuildsState> PendingQueryState = MakeShared<FListBuildsState>();
-		PendingQueryState->PendingQueries = NamespacePlatformBucketTuples.Num();
-		PendingQueryState->QueryState.SetNum(NamespacePlatformBucketTuples.Num());
+		PendingQueryState->PendingQueries = NamespaceBucketTuples.Num();
+		PendingQueryState->QueryState.SetNum(NamespaceBucketTuples.Num());
 		uint32 QueryIndex = 0;
 
 		++BuildRefreshGeneration;
 		BuildGroups.Empty();
 
-		for (const FNamespacePlatformBucketTuple& NamespacePlatformBucket : NamespacePlatformBucketTuples)
+		for (const FNamespacePlatformBucketTuple& NamespacePlatformBucket : NamespaceBucketTuples)
 		{
-			ServiceInstance->ListBuilds(NamespacePlatformBucket.Namespace, NamespacePlatformBucket.Bucket,
-			[this, QueryIndex, Namespace = NamespacePlatformBucket.Namespace, Platform = NamespacePlatformBucket.Platform,
+			ServiceInstance->ListBuilds(NamespacePlatformBucket.Namespace, NamespacePlatformBucket.BucketPrefix,
+			[this, QueryIndex, Namespace = NamespacePlatformBucket.Namespace,
 				ExpectedBuildRefreshGeneration = BuildRefreshGeneration.load(), PendingQueryState]
 			(TArray<FBuildServiceInstance::FBuildRecord>&& Results) mutable
 				{
 					FBuildState& NewBuildState = PendingQueryState->QueryState[QueryIndex];
 					NewBuildState.Namespace = MoveTemp(Namespace);
-					NewBuildState.Platform = MoveTemp(Platform);
 					NewBuildState.Results = MoveTemp(Results);
 
 					if (--PendingQueryState->PendingQueries == 0)
@@ -384,7 +386,7 @@ void SBuildSelection::RebuildLists()
 			++QueryIndex;
 		}
 
-		if (NamespacePlatformBucketTuples.IsEmpty())
+		if (NamespaceBucketTuples.IsEmpty())
 		{
 			BuildListRefreshesInProgress.fetch_sub(1);
 		}
@@ -458,6 +460,15 @@ void SBuildSelection::RegenerateBuildGroups(UE::BuildSelection::Internal::FListB
 				BuildGroup = MakeShared<FBuildGroup>();
 			}
 
+			FString Platform;
+			if (FCbFieldView PlatformField = BuildRecord.Metadata["platform"]; PlatformField.HasValue() && !PlatformField.HasError())
+			{
+				if (PlatformField.IsString())
+				{
+					Platform = FUTF8ToTCHAR(PlatformField.AsString());
+				}
+			}
+
 			if (BuildGroup->DisplayName.IsEmpty())
 			{
 				FString Job;
@@ -536,7 +547,7 @@ void SBuildSelection::RegenerateBuildGroups(UE::BuildSelection::Internal::FListB
 				BuildGroup->Job = Job;
 			}
 
-			BuildGroup->PerPlatformBuilds.FindOrAdd(BuildState.Platform, MoveTemp(BuildRecord));
+			BuildGroup->PerPlatformBuilds.FindOrAdd(Platform, MoveTemp(BuildRecord));
 		}
 	}
 
