BuildStorageTool usability improvements:

- ship and install/update/start zenserver as well as zen executables (and crash handler and symbols) when using BuildServerInterface functionality
- display user-friendly localized names for the "Build Type" field which is detected based on regex pattern match of the bucket segment.  Notably: oplog => Cooked Data
- save user selected stream/branch, project, and build type values into a user-wide setting
- select default destination based saved knowledge of last used workspace root, or from the first known workspace root on the machine

diff --git a/Engine/Source/Developer/Zen/Private/BuildServerInterface.cpp b/Engine/Source/Developer/Zen/Private/BuildServerInterface.cpp
--- a/Engine/Source/Developer/Zen/Private/BuildServerInterface.cpp
+++ b/Engine/Source/Developer/Zen/Private/BuildServerInterface.cpp
@@ -597,6 +597,7 @@ FBuildServiceInstance::Initialize()
 #endif
 	ISocketSubsystem::Get();
 	FDesktopPlatformModule::TryGet();
+	ZenService = MakePimpl<FScopeZenService>();
 }
 
 bool
diff --git a/Engine/Source/Developer/Zen/Public/Experimental/BuildServerInterface.h b/Engine/Source/Developer/Zen/Public/Experimental/BuildServerInterface.h
--- a/Engine/Source/Developer/Zen/Public/Experimental/BuildServerInterface.h
+++ b/Engine/Source/Developer/Zen/Public/Experimental/BuildServerInterface.h
@@ -14,6 +14,7 @@
 #include "HAL/Thread.h"
 #include "Serialization/CompactBinary.h"
 #include "Templates/Function.h"
+#include "Templates/PimplPtr.h"
 #include "Templates/SharedPointer.h"
 #include "Templates/UniquePtr.h"
 #include <atomic>
@@ -21,6 +22,7 @@
 #define UE_API ZEN_API
 
 namespace UE { class IHttpClient; }
+namespace UE::Zen { class FScopeZenService; }
 namespace UE::Zen::Build::Private { class FAccessToken; }
 namespace UE::Zen::Build::Private { struct FBuildTransferState; }
 
@@ -139,6 +141,8 @@ private:
 	FRWLock NamespacesAndBucketsLock;
 	TMultiMap<FString, FString> NamespacesAndBuckets;
 
+	TPimplPtr<FScopeZenService> ZenService;
+
 	TMpscQueue<TSharedRef<Private::FBuildTransferState>> BuildTransferThreadStates;
 	FThread BuildTransferThread;
 	std::atomic<bool> bBuildTransferThreadStarting = false;
diff --git a/Engine/Source/Editor/StorageServerWidgets/Internal/SBuildSelection.h b/Engine/Source/Editor/StorageServerWidgets/Internal/SBuildSelection.h
--- a/Engine/Source/Editor/StorageServerWidgets/Internal/SBuildSelection.h
+++ b/Engine/Source/Editor/StorageServerWidgets/Internal/SBuildSelection.h
@@ -4,6 +4,7 @@
 
 #include "Experimental/BuildServerInterface.h"
 #include "Experimental/ZenServerInterface.h"
+#include "Internationalization/Regex.h"
 #include "Templates/SharedPointer.h"
 #include "Widgets/DeclarativeSyntaxSupport.h"
 #include "Widgets/Input/SComboBox.h"
@@ -18,6 +19,7 @@
 
 #define UE_API STORAGESERVERWIDGETS_API
 
+class FUProjectDictionary;
 class SMultiSelectComboBox;
 
 namespace UE::BuildSelection::Internal{ struct FListBuildsState; }
@@ -64,6 +66,12 @@ private:
 		Unknown,
 		Count
 	};
+	struct FKnownBuildType
+	{
+		EBuildType Type;
+		FText UserText;
+	};
+	static TPair<FRegexPattern, FKnownBuildType> MakeKnownBuildTypePattern(const FStringView InPattern, EBuildType Type, FText&& InText);
 	UE_API EBuildType GetSelectedBuildType() const;
 	UE_API void RebuildLists();
 	UE_API void RegenerateBuildGroups(UE::BuildSelection::Internal::FListBuildsState& ListBuildsState);
@@ -71,6 +79,7 @@ private:
 	UE_API void RegenerateActivePlatformFilters();
 	UE_API void ValidateBuildGroupSelection();
 	UE_API TSharedRef<SWidget> OnGenerateTextBlockFromString(TSharedPtr<FString> Item);
+	UE_API TSharedRef<SWidget> OnGenerateBuildTypeTextBlockFromString(TSharedPtr<FString> Item);
 	UE_API bool BuildGroupIsSelectableOrNavigable(TSharedPtr<FBuildGroup> InItem) const;
 	UE_API TSharedRef<ITableRow> GenerateBuildGroupRow(TSharedPtr<FBuildGroup> InItem, const TSharedRef<STableViewBase>& InOwningTable);
 	UE_API TSharedPtr<SWidget> OnGetBuildGroupContextMenuContent() const;
@@ -81,23 +90,45 @@ private:
 	EColumnSortMode::Type GetColumnSortMode(const FName ColumnId) const;
 	void OnColumnSortModeChanged(const EColumnSortPriority::Type SortPriority, const FName& ColumnId, const EColumnSortMode::Type InSortMode);
 
+	FText ConvertBuildTypeToText(const FString& InBuildType);
 
-	UE_API FString& GetUserSelectedDestination();
+	UE_API void SetUserSelectedDestination(const FStringView InDestination);
 	UE_API FString GetUserSelectedDestination() const;
 	UE_API FString GetDefaultDestination() const;
 	UE_API FString GetEffectiveDestination() const;
 
+	void SetSelectedStream(const TSharedPtr<FString> InSelectedStream);
+	void SetSelectedProject(const TSharedPtr<FString> InSelectedProject);
+	void SetSelectedBuildType(const TSharedPtr<FString> InBuildType);
+
+	bool WriteSetting(FStringView InSectionName, FStringView InKeyName, FStringView InValue);
+	bool WriteSetting(FStringView InKeyName, FStringView InValue);
+	bool ReadSetting(FStringView InSectionName, FStringView InKeyName, FString& OutValue);
+	bool ReadSetting(FStringView InKeyName, FString& OutValue);
+
+	void SetUserSelectedProjectDictionaryRoot(FStringView InRoot);
+
 	static UE_API FString SanitizeForPath(const FString& InString);
 	static UE_API FString SanitizeForZenId(const FString& InString);
 	static UE_API FString SanitizeBucketSegment(const FString& InString);
 
+	static const TCHAR* LexToString(EBuildType BuildType);
+
 	UE_API FReply ExploreDestination_OnClicked();
 
+	static const FString EmptyString;
+
+	TMap<FString, FString> EngineInstallations;
+
+	TUniquePtr<FUProjectDictionary> FallbackProjectDictionary;
+
+	FString UserSelectedProjectDictionaryRootDir;
+	TUniquePtr<FUProjectDictionary> UserSelectedProjectDictionary;
+
 	FString BuildUrl;
 	FString BuildId;
 	FString BuildVersion;
 	FString BranchName;
-	FString BucketFormattedBranchName;
 	int EffectiveCompatibleChangelist = 0;
 
 	TSharedPtr<SComboBox<TSharedPtr<FString>>> ProjectWidget;
@@ -124,6 +155,8 @@ private:
 	TSharedPtr<FString> SelectedBuildType;
 	FString SelectedCommitFilter;
 
+	TArray<TPair<FRegexPattern, FKnownBuildType>> KnownBuildTypePatterns;
+
 	TSharedPtr<SListView<TSharedPtr<SBuildSelection::FBuildGroup, ESPMode::ThreadSafe>>> BuildGroupListView;
 	FName BuildGroupSortByColumn;
 	EColumnSortMode::Type BuildGroupSortMode;
diff --git a/Engine/Source/Editor/StorageServerWidgets/Private/SBuildSelection.cpp b/Engine/Source/Editor/StorageServerWidgets/Private/SBuildSelection.cpp
--- a/Engine/Source/Editor/StorageServerWidgets/Private/SBuildSelection.cpp
+++ b/Engine/Source/Editor/StorageServerWidgets/Private/SBuildSelection.cpp
@@ -31,6 +31,8 @@
 
 #define LOCTEXT_NAMESPACE "StorageServerBuild"
 
+const FString SBuildSelection::EmptyString;
+
 namespace UE::BuildSelection::Internal
 {
 	namespace FBuildGroupIds
@@ -61,10 +63,57 @@ namespace UE::BuildSelection::Internal
 	};
 }
 
+TPair<FRegexPattern, SBuildSelection::FKnownBuildType> SBuildSelection::MakeKnownBuildTypePattern(const FStringView InPattern, EBuildType Type, FText&& InText)
+{
+	FString Pattern(InPattern);
+	return TPair<FRegexPattern, FKnownBuildType>(Pattern, FKnownBuildType{Type, InText});
+}
+
 void SBuildSelection::Construct(const FArguments& InArgs)
 {
 	using namespace UE::BuildSelection::Internal;
 
+	KnownBuildTypePatterns = {
+		// These should be made more constinent and perhaps specified in INI
+		MakeKnownBuildTypePattern(TEXT(".*oplog-?(.*)"), EBuildType::Oplog, LOCTEXT("BuildSelection_BuildType_Oplog", "Cooked Data") ),
+		MakeKnownBuildTypePattern(TEXT(".*packaged-?build-?(.*)"), EBuildType::PackagedBuild, LOCTEXT("BuildSelection_BuildType_PackagedBuild", "Packaged Build")),
+		MakeKnownBuildTypePattern(TEXT(".*staged-?(.*?)-?build.*"), EBuildType::StagedBuild, LOCTEXT("BuildSelection_BuildType_StagedBuild", "Staged Build")),
+		MakeKnownBuildTypePattern(TEXT(".*ugs-?pcb-?(.*)"), EBuildType::EditorPreCompiledBinary, LOCTEXT("BuildSelection_BuildType_PreCompiledBinary", "Editor Pre-compiled Binary")),
+		MakeKnownBuildTypePattern(TEXT(".*installed-?build-?(.*)"), EBuildType::EditorPreCompiledBinary, LOCTEXT("BuildSelection_BuildType_PreCompiledBinary", "Editor Pre-compiled Binary"))
+	};
+
+	FString UserSelectedEngineInstallation;
+	if (ReadSetting(TEXT("UserSelectedEngineInstallation"), UserSelectedEngineInstallation))
+	{
+		if (IDesktopPlatform* DesktopPlatform = FDesktopPlatformModule::Get())
+		{
+			FString UserSelectedRootDir;
+			if (DesktopPlatform->GetEngineRootDirFromIdentifier(UserSelectedEngineInstallation, UserSelectedRootDir))
+			{
+				SetUserSelectedProjectDictionaryRoot(UserSelectedRootDir);
+			}
+		}
+	}
+
+	for (int32 BuildTypeIndex = 0; BuildTypeIndex < (int32)EBuildType::Count; ++BuildTypeIndex)
+	{
+		ReadSetting(LexToString((EBuildType)BuildTypeIndex), TEXT("UserSelectedDestination"), UserSelectedDestinations[BuildTypeIndex]);
+	}
+
+	if (EngineInstallations.IsEmpty())
+	{
+		if (IDesktopPlatform* DesktopPlatform = FDesktopPlatformModule::Get())
+		{
+			DesktopPlatform->EnumerateEngineInstallations(EngineInstallations);
+		}
+	}
+
+	for (const TPair<FString, FString>& EngineInstallation : EngineInstallations)
+	{
+		FallbackProjectDictionary = MakeUnique<FUProjectDictionary>(EngineInstallation.Value);
+		break;
+	}
+
 	ZenServiceInstance = InArgs._ZenServiceInstance;
 	BuildServiceInstance = InArgs._BuildServiceInstance;
 	OnBuildTransferStarted = InArgs._OnBuildTransferStarted;
@@ -78,12 +127,16 @@ void SBuildSelection::Construct(const FArguments& InArgs)
 		BuildUrl = CurrentBuildVersion.BuildUrl;
 		BuildId = CurrentBuildVersion.BuildId;
 		BuildVersion = CurrentBuildVersion.BuildVersion;
-		BranchName = CurrentBuildVersion.BranchName;
-		BucketFormattedBranchName = SanitizeBucketSegment(BranchName);
+		BranchName = SanitizeBucketSegment(CurrentBuildVersion.BranchName);
 
 		EffectiveCompatibleChangelist = CurrentBuildVersion.GetEffectiveCompatibleChangelist();
 	}
 
+	if (BranchName.IsEmpty())
+	{
+		ReadSetting(TEXT("Branch"), BranchName);
+	}
+
 	for (bool& bAppendBuildNameToDestination : bAppendBuildNameToDestinations)
 	{
 		bAppendBuildNameToDestination = false;
@@ -115,25 +168,13 @@ SBuildSelection::EBuildType SBuildSelection::GetSelectedBuildType() const
 	{
 		return EBuildType::Unknown;
 	}
-	else if (SelectedBuildType->Contains(TEXT("oplog")))
-	{
-		return EBuildType::Oplog;
-	}
-	else if (SelectedBuildType->Contains(TEXT("stagedbuild")) || SelectedBuildType->Contains(TEXT("staged-build")))
-	{
-		return EBuildType::StagedBuild;
-	}
-	else if (SelectedBuildType->Contains(TEXT("packagedbuild")) || SelectedBuildType->Contains(TEXT("packaged-build")))
+	for (const TPair<FRegexPattern, FKnownBuildType>& KnownBuildTypeItem : KnownBuildTypePatterns)
 	{
-		return EBuildType::PackagedBuild;
-	}
-	else if (SelectedBuildType->Contains(TEXT("ugs-pcb")))
-	{
-		return EBuildType::EditorPreCompiledBinary;
-	}
-	else if (SelectedBuildType->Contains(TEXT("installedbuild")) || SelectedBuildType->Contains(TEXT("installed-build")))
-	{
-		return EBuildType::EditorInstalledBuild;
+		FRegexMatcher Matcher(KnownBuildTypeItem.Key, *SelectedBuildType);
+		if (Matcher.FindNext())
+		{
+			return KnownBuildTypeItem.Value.Type;
+		}
 	}
 	return EBuildType::Unknown;
 }
@@ -222,7 +263,7 @@ void SBuildSelection::RebuildLists()
 			if (BucketSegmentViews.Num() == SegmentIndexNum)
 			{
 				int32 CurrentIndex = Streams.AddUnique(FString(BucketSegmentViews[SegmentIndexStream]));
-				if (!BranchName.IsEmpty() && (Streams[CurrentIndex] == BucketFormattedBranchName))
+				if (!BranchName.IsEmpty() && (Streams[CurrentIndex] == BranchName))
 				{
 					DefaultStreamIndex = CurrentIndex;
 				}
@@ -237,7 +278,7 @@ void SBuildSelection::RebuildLists()
 			});
 		if (!SelectedStream && DefaultStreamListItem)
 		{
-			SelectedStream = *DefaultStreamListItem;
+			SetSelectedStream(*DefaultStreamListItem);
 		}
 		ConformSelection(SelectedStream, StreamList);
 
@@ -255,6 +296,14 @@ void SBuildSelection::RebuildLists()
 				Projects.AddUnique(FString(BucketSegmentViews[SegmentIndexProject]));
 			}
 		}
+
+		if (!SelectedProject)
+		{
+			if (FString PastProject; ReadSetting(TEXT("Project"), PastProject))
+			{
+				SelectedProject = MakeShared<FString>(MoveTemp(PastProject));
+			}
+		}
 		ConvertToSharedPtrs(Projects, ProjectList);
 		ConformSelection(SelectedProject, ProjectList);
 
@@ -277,6 +326,14 @@ void SBuildSelection::RebuildLists()
 				BuildTypes.AddUnique(FString(BucketSegmentViews[SegmentIndexBuildType]));
 			}
 		}
+
+		if (!SelectedBuildType)
+		{
+			if (FString PastBuildType; ReadSetting(TEXT("BuildType"), PastBuildType))
+			{
+				SelectedBuildType = MakeShared<FString>(MoveTemp(PastBuildType));
+			}
+		}
 		ConvertToSharedPtrs(BuildTypes, BuildTypeList);
 		ConformSelection(SelectedBuildType, BuildTypeList);
 
@@ -676,6 +733,12 @@ TSharedRef<SWidget> SBuildSelection::OnGenerateTextBlockFromString(TSharedPtr<FS
 		.Text(FText::FromString(*Item));
 }
 
+TSharedRef<SWidget> SBuildSelection::OnGenerateBuildTypeTextBlockFromString(TSharedPtr<FString> Item)
+{
+	return SNew(STextBlock)
+		.Text(Item ? ConvertBuildTypeToText(*Item) : FText::GetEmpty());
+}
+
 bool SBuildSelection::BuildGroupIsSelectableOrNavigable(FBuildSelectionBuildGroupPtr InItem) const
 {
 	if (!InItem)
@@ -888,7 +951,7 @@ void SBuildSelection::BuildGroupSelectionChanged(FBuildSelectionBuildGroupPtr It
 				]
 			];
 
-		if (++Row > 2)
+		if (++Row > 3)
 		{
 			Column++;
 			Row = 0;
@@ -919,11 +982,53 @@ void SBuildSelection::OnColumnSortModeChanged(const EColumnSortPriority::Type So
 		});
 }
 
-FString& SBuildSelection::GetUserSelectedDestination()
+FText SBuildSelection::ConvertBuildTypeToText(const FString& InBuildType)
+{
+	for (const TPair<FRegexPattern, FKnownBuildType>& KnownBuildTypeItem : KnownBuildTypePatterns)
+	{
+		FRegexMatcher Matcher(KnownBuildTypeItem.Key, InBuildType);
+		if (Matcher.FindNext())
+		{
+			FString OptionalCaptureGroupString = Matcher.GetCaptureGroup(1);
+			if (OptionalCaptureGroupString.IsEmpty())
+			{
+				return KnownBuildTypeItem.Value.UserText;
+			}
+			return FText::Format(LOCTEXT("BuildSelection_KnownBuildTypeWithCaptureGroupFormat", "{0} ({1})"), KnownBuildTypeItem.Value.UserText, FText::FromString(OptionalCaptureGroupString));
+		}
+	}
+
+	return FText::FromString(InBuildType);
+}
+
+void SBuildSelection::SetUserSelectedDestination(const FStringView InDestination)
 {
 	EBuildType BuildType = GetSelectedBuildType();
 	int BuildTypeIndex = (int)BuildType;
-	return UserSelectedDestinations[BuildTypeIndex];
+	UserSelectedDestinations[BuildTypeIndex] = InDestination;
+
+	if (BuildType != EBuildType::Oplog)
+	{
+		if (EngineInstallations.IsEmpty())
+		{
+			if (IDesktopPlatform* DesktopPlatform = FDesktopPlatformModule::Get())
+			{
+				DesktopPlatform->EnumerateEngineInstallations(EngineInstallations);
+			}
+		}
+
+		for (const TPair<FString, FString>& EngineInstallation : EngineInstallations)
+		{
+			if (FPaths::IsUnderDirectory(UserSelectedDestinations[BuildTypeIndex], EngineInstallation.Value))
+			{
+				WriteSetting(TEXT("UserSelectedEngineInstallation"), EngineInstallation.Key);
+				SetUserSelectedProjectDictionaryRoot(EngineInstallation.Value);
+				break;
+			}
+		}
+	}
+
+	WriteSetting(LexToString(BuildType), TEXT("UserSelectedDestination"), InDestination);
 }
 
 FString SBuildSelection::GetUserSelectedDestination() const
@@ -939,6 +1044,12 @@ FString SBuildSelection::GetDefaultDestination() const
 	IPlatformFile& PlatformFile = FPlatformFileManager::Get().GetPlatformFile();
 
 	FString ProjectFilename = FUProjectDictionary::GetDefault().GetProjectPathForGame(**SelectedProject);
+
+	if (ProjectFilename.IsEmpty() && UserSelectedProjectDictionary)
+	{
+		ProjectFilename = UserSelectedProjectDictionary->GetProjectPathForGame(**SelectedProject);
+	}
+
 	if (ProjectFilename.IsEmpty())
 	{
 		if (BuildType == EBuildType::Oplog)
@@ -951,6 +1062,11 @@ FString SBuildSelection::GetDefaultDestination() const
 			);
 	}
 
+	if (ProjectFilename.IsEmpty() && FallbackProjectDictionary)
+	{
+		ProjectFilename = FallbackProjectDictionary->GetProjectPathForGame(**SelectedProject);
+	}
+
 	switch (BuildType)
 	{
 	case EBuildType::StagedBuild:
@@ -975,6 +1091,69 @@ FString SBuildSelection::GetEffectiveDestination() const
 	return UserSelectedDestination.IsEmpty() ? GetDefaultDestination() : UserSelectedDestination;
 }
 
+void SBuildSelection::SetSelectedStream(const TSharedPtr<FString> InSelectedStream)
+{
+	SelectedStream = InSelectedStream;
+	WriteSetting(TEXT("Branch"), *SelectedStream);
+}
+
+void SBuildSelection::SetSelectedProject(const TSharedPtr<FString> InSelectedProject)
+{
+	SelectedProject = InSelectedProject;
+	WriteSetting(TEXT("Project"), *SelectedProject);
+}
+
+void SBuildSelection::SetSelectedBuildType(const TSharedPtr<FString> InBuildType)
+{
+	SelectedBuildType = InBuildType;
+	WriteSetting(TEXT("BuildType"), *SelectedBuildType);
+}
+
+bool SBuildSelection::WriteSetting(const FStringView InSectionName, const FStringView InKeyName, const FStringView InValue)
+{
+	static const FString StoreId = TEXT("Epic Games");
+	static const FString SectionName = TEXT("Unreal Engine/Build Storage");
+	const FString FinalSectionName = InSectionName.IsEmpty() ? SectionName : FPaths::Combine(SectionName, InSectionName);
+	const FString FinalKeyName(InKeyName);
+	FString Value(InValue);
+	return FPlatformMisc::SetStoredValue(TEXT("Epic Games"), FinalSectionName, FinalKeyName, Value);
+}
+
+bool SBuildSelection::WriteSetting(const FStringView InKeyName, const FStringView InValue)
+{
+	return WriteSetting(EmptyString, InKeyName, InValue);
+}
+
+bool SBuildSelection::ReadSetting(const FStringView InSectionName, const FStringView InKeyName, FString& OutValue)
+{
+	static const FString StoreId = TEXT("Epic Games");
+	static const FString SectionName = TEXT("Unreal Engine/Build Storage");
+	const FString FinalSectionName = InSectionName.IsEmpty() ? SectionName : FPaths::Combine(SectionName, InSectionName);
+	const FString FinalKeyName(InKeyName);
+	return FPlatformMisc::GetStoredValue(TEXT("Epic Games"), FinalSectionName, FinalKeyName, OutValue);
+}
+
+bool SBuildSelection::ReadSetting(const FStringView InKeyName, FString& OutValue)
+{
+	return ReadSetting(EmptyString, InKeyName, OutValue);
+}
+
+void SBuildSelection::SetUserSelectedProjectDictionaryRoot(const FStringView InRoot)
+{
+	if (InRoot.IsEmpty())
+	{
+		UserSelectedProjectDictionaryRootDir.Empty();
+		UserSelectedProjectDictionary.Reset();
+	}
+	FString Root(InRoot);
+	FString FullRoot = FPaths::ConvertRelativePathToFull(Root);
+	if (!FPaths::IsSamePath(FullRoot, UserSelectedProjectDictionaryRootDir))
+	{
+		UserSelectedProjectDictionaryRootDir = FullRoot;
+		UserSelectedProjectDictionary = MakeUnique<FUProjectDictionary>(UserSelectedProjectDictionaryRootDir);
+	}
+}
+
 FString SBuildSelection::SanitizeForPath(const FString& InString)
 {
 	// TODO: Had to remove path sanitization as the engine only provides validation
@@ -1035,6 +1214,22 @@ FString SBuildSelection::SanitizeBucketSegment(const FString& InString)
 	return OutString;
 }
 
+const TCHAR* SBuildSelection::LexToString(EBuildType BuildType)
+{
+	static const TCHAR* Strings[] =
+	{
+		TEXT("Oplog"),
+		TEXT("StagedBuild"),
+		TEXT("PackagedBuild"),
+		TEXT("EditorPreCompiledBinary"),
+		TEXT("EditorInstalledBuild"),
+		TEXT("Unknown")
+	};
+	static_assert((int32)(EBuildType::Count) == UE_ARRAY_COUNT(Strings), "SBuildSelection::LexToString must contain a string for each member of EBuildType");
+
+	return Strings[(int32)BuildType];
+}
+
 FReply SBuildSelection::ExploreDestination_OnClicked()
 {
 	if (GetSelectedBuildType() == EBuildType::Oplog)
@@ -1060,17 +1255,19 @@ FReply SBuildSelection::ExploreDestination_OnClicked()
 
 void SBuildSelection::OnOpenDestinationDirectoryClicked()
 {
-	bool bDirectorySelected = false;
-
 	if (IDesktopPlatform* DesktopPlatform = FDesktopPlatformModule::Get())
 	{
 		const FString Title = LOCTEXT("BuildSelection_DestinationDirectoryBrowserTitle", "Choose destination directory").ToString();
 
-		bDirectorySelected = DesktopPlatform->OpenDirectoryDialog(
+		FString NewDestination = GetUserSelectedDestination();
+		if (DesktopPlatform->OpenDirectoryDialog(
 			FSlateApplication::Get().FindBestParentWindowHandleForDialogs(nullptr),
 			Title,
 			GetEffectiveDestination(),
-			GetUserSelectedDestination());
+			NewDestination))
+		{
+			SetUserSelectedDestination(NewDestination);
+		}
 	}
 }
 
@@ -1108,7 +1305,7 @@ TSharedRef<SWidget> SBuildSelection::GetBuildDestinationPanel()
 			})
 			.OnTextCommitted_Lambda([this](const FText& Text, const ETextCommit::Type CommitType)
 			{
-				GetUserSelectedDestination() = Text.ToString();
+				SetUserSelectedDestination(Text.ToString());
 			})
 		]
 		+ SHorizontalBox::Slot()
@@ -1157,7 +1354,7 @@ TSharedRef<SWidget> SBuildSelection::GetBuildDestinationPanel()
 			})
 			.OnTextCommitted_Lambda([this](const FText& Text, const ETextCommit::Type CommitType)
 			{
-				GetUserSelectedDestination() = Text.ToString();
+				SetUserSelectedDestination(Text.ToString());
 			})
 		]
 	]
@@ -1289,7 +1486,7 @@ TSharedRef<SWidget> SBuildSelection::GetGridPanel()
 				{
 					if (Item.IsValid() && *Item != *SelectedStream)
 					{
-						SelectedStream = Item;
+						SetSelectedStream(Item);
 						RebuildLists();
 					}
 				})
@@ -1337,7 +1534,7 @@ TSharedRef<SWidget> SBuildSelection::GetGridPanel()
 				{
 					if (Item.IsValid() && *Item != *SelectedProject)
 					{
-						SelectedProject = Item;
+						SetSelectedProject(Item);
 						RebuildLists();
 					}
 				})
@@ -1385,17 +1582,21 @@ TSharedRef<SWidget> SBuildSelection::GetGridPanel()
 				{
 					if (Item.IsValid() && *Item != *SelectedBuildType)
 					{
-						SelectedBuildType = Item;
+						SetSelectedBuildType(Item);
 						RebuildLists();
 					}
 				})
-				.OnGenerateWidget(this, &SBuildSelection::OnGenerateTextBlockFromString)
+				.OnGenerateWidget(this, &SBuildSelection::OnGenerateBuildTypeTextBlockFromString)
 				[
 					SNew(STextBlock)
 						.MinDesiredWidth(MinDesiredWidth)
 						.Text_Lambda([this]()
 						{
-							return FText::FromString(SelectedBuildType ? **SelectedBuildType : TEXT(""));
+							if (!SelectedBuildType)
+							{
+								return FText::GetEmpty();
+							}
+							return ConvertBuildTypeToText(*SelectedBuildType);
 						})
 				]
 			]
diff --git a/Engine/Source/Programs/BuildStorageTool/BuildStorageTool.Build.cs b/Engine/Source/Programs/BuildStorageTool/BuildStorageTool.Build.cs
--- a/Engine/Source/Programs/BuildStorageTool/BuildStorageTool.Build.cs
+++ b/Engine/Source/Programs/BuildStorageTool/BuildStorageTool.Build.cs
@@ -38,17 +38,24 @@ public class BuildStorageTool : ModuleRules
 
 		if (Target.Platform == UnrealTargetPlatform.Win64)
 		{
+			RuntimeDependencies.Add("$(EngineDir)/Binaries/Win64/crashpad_handler.exe");
 			RuntimeDependencies.Add("$(EngineDir)/Binaries/Win64/zen.exe");
+			RuntimeDependencies.Add("$(EngineDir)/Binaries/Win64/zen.pdb");
+			RuntimeDependencies.Add("$(EngineDir)/Binaries/Win64/zenserver.exe");
+			RuntimeDependencies.Add("$(EngineDir)/Binaries/Win64/zenserver.pdb");
 			RuntimeDependencies.Add("$(EngineDir)/Binaries/DotNET/OidcToken/win-x64/OidcToken.exe");
 		}
 		else if (Target.Platform == UnrealTargetPlatform.Linux)
 		{
 			RuntimeDependencies.Add("$(EngineDir)/Binaries/Linux/zen");
+			RuntimeDependencies.Add("$(EngineDir)/Binaries/Linux/zenserver");
 			RuntimeDependencies.Add("$(EngineDir)/Binaries/DotNET/OidcToken/linux-x64/OidcToken");
 		}
 		else if (Target.Platform == UnrealTargetPlatform.Mac)
 		{
+			RuntimeDependencies.Add("$(EngineDir)/Binaries/Mac/crashpad_handler");
 			RuntimeDependencies.Add("$(EngineDir)/Binaries/Mac/zen");
+			RuntimeDependencies.Add("$(EngineDir)/Binaries/Mac/zenserver");
 			RuntimeDependencies.Add("$(EngineDir)/Binaries/DotNET/OidcToken/osx-x64/OidcToken");
 		}
 
diff --git a/Engine/Source/Programs/BuildStorageTool/BuildStorageTool.Target.cs b/Engine/Source/Programs/BuildStorageTool/BuildStorageTool.Target.cs
--- a/Engine/Source/Programs/BuildStorageTool/BuildStorageTool.Target.cs
+++ b/Engine/Source/Programs/BuildStorageTool/BuildStorageTool.Target.cs
@@ -21,7 +21,7 @@ public class BuildStorageToolTarget : TargetRules
 		bCompileNavmeshClusterLinks		= false;
 		bCompileNavmeshSegmentLinks		= false;
 		bCompileRecast					= false;
-		bCompileICU 					= false;
+		bCompileICU 					= true; // Required for regex support
 		bWithLiveCoding					= false;
 		bBuildDeveloperTools			= false;
 		bBuildWithEditorOnlyData		= false;
diff --git a/Engine/Source/Programs/BuildStorageTool/Private/BuildStorageTool.cpp b/Engine/Source/Programs/BuildStorageTool/Private/BuildStorageTool.cpp
--- a/Engine/Source/Programs/BuildStorageTool/Private/BuildStorageTool.cpp
+++ b/Engine/Source/Programs/BuildStorageTool/Private/BuildStorageTool.cpp
@@ -152,7 +152,7 @@ public:
 		Window =
 			SNew(SWindow)
 			.Title(GetWindowTitle())
-			.ClientSize(FVector2D(900.0f, 600.0f))
+			.ClientSize(FVector2D(1000.0f, 600.0f))
 			.ActivationPolicy(EWindowActivationPolicy::Always)
 			.SizingRule(ESizingRule::UserSized)
 			.IsTopmostWindow(false)
