Add right click context menu item for activities to open the destination.  Also add double-clicking of actvity row to open the destination for that row.

diff --git a/Engine/Source/Developer/Zen/Private/BuildServerInterface.cpp b/Engine/Source/Developer/Zen/Private/BuildServerInterface.cpp
--- a/Engine/Source/Developer/Zen/Private/BuildServerInterface.cpp
+++ b/Engine/Source/Developer/Zen/Private/BuildServerInterface.cpp
@@ -46,12 +46,6 @@ DEFINE_LOG_CATEGORY_STATIC(LogBuildServiceInstance, Log, All);
 namespace Private
 {
 
-enum class EBuildTransferType
-{
-	Files,
-	Oplog
-};
-
 struct FProgressState
 {
 	FString Label;
@@ -61,7 +55,7 @@ struct FProgressState
 
 struct FBuildTransferState : public TSharedFromThis<FBuildTransferState, ESPMode::ThreadSafe>
 {
-	EBuildTransferType Type;
+	FBuildServiceInstance::FBuildTransfer::EType Type;
 	FCbObjectId BuildId;
 	FString Name;
 	FString Namespace;
@@ -413,7 +407,7 @@ FBuildServiceInstance::FBuildTransfer
 FBuildServiceInstance::StartBuildTransfer(const FCbObjectId& BuildId, FStringView Name, FStringView DestinationFolder, FStringView Namespace, FStringView Bucket)
 {
 	TSharedPtr<Private::FBuildTransferState> NewState = MakeShared<Private::FBuildTransferState>();
-	NewState->Type = Private::EBuildTransferType::Files;
+	NewState->Type = FBuildTransfer::EType::Files;
 	NewState->BuildId = BuildId;
 	NewState->Name = Name;
 	NewState->Namespace = Namespace;
@@ -431,7 +425,7 @@ FBuildServiceInstance::FBuildTransfer
 FBuildServiceInstance::StartOplogBuildTransfer(const FCbObjectId& BuildId, FStringView Name, FStringView DestinationProjectId, FStringView DestinationOplogId, FStringView ProjectFilePath, FStringView Namespace, FStringView Bucket, FName TargetPlatformName)
 {
 	TSharedPtr<Private::FBuildTransferState> NewState = MakeShared<Private::FBuildTransferState>();
-	NewState->Type = Private::EBuildTransferType::Oplog;
+	NewState->Type = FBuildTransfer::EType::Oplog;
 	NewState->BuildId = BuildId;
 	NewState->Name = Name;
 	NewState->Namespace = Namespace;
@@ -472,6 +466,25 @@ FBuildServiceInstance::RepeatBuildTransfer(FBuildTransfer BuildTransfer)
 	return NewTransfer;
 }
 
+FBuildServiceInstance::FBuildTransfer::EType
+FBuildServiceInstance::FBuildTransfer::GetType() const
+{
+	if (!State)
+	{
+		return EType::Count;
+	}
+	return State->Type;
+}
+
+FString FBuildServiceInstance::FBuildTransfer::GetDestination() const
+{
+	if (!State)
+	{
+		return FString();
+	}
+	return State->Destination;
+}
+
 FString
 FBuildServiceInstance::FBuildTransfer::GetDescription() const
 {
@@ -871,13 +884,13 @@ FBuildServiceInstance::BuildTransferThreadLoop()
 		FString CommandLineArgs;
 		switch (State->Type)
 		{
-		case Private::EBuildTransferType::Files:
+		case FBuildTransfer::EType::Files:
 			DestinationDirectory = FPaths::ConvertRelativePathToFull(State->Destination);
 			CommandLineArgs = FString::Printf(TEXT("builds download --log-progress --host %s --local-path \"%s\" --namespace %s --bucket %s --build-id %s --oidctoken-exe-path \"%s\""),
 				*WriteToString<64>(EffectiveDomain), *DestinationDirectory, *State->Namespace, *State->Bucket,
 				*WriteToString<64>(State->BuildId), *OidcTokenExecutableFilename);
 			break;
-		case Private::EBuildTransferType::Oplog:
+		case FBuildTransfer::EType::Oplog:
 			{
 				FString DestinationProjectId, DestinationOplogId;
 				FString SplitDelim(TEXT("/"));
@@ -1075,7 +1088,7 @@ FBuildServiceInstance::BuildTransferThreadLoop()
 		{
 			FReadScopeLock _(State->Lock);
 			bCleanDestination = (State->Status == EBuildTransferStatus::Succeeded) &&
-				(State->Type == Private::EBuildTransferType::Oplog) &&
+				(State->Type == FBuildTransfer::EType::Oplog) &&
 				!DestinationDirectory.IsEmpty();
 		}
 
diff --git a/Engine/Source/Developer/Zen/Public/Experimental/BuildServerInterface.h b/Engine/Source/Developer/Zen/Public/Experimental/BuildServerInterface.h
--- a/Engine/Source/Developer/Zen/Public/Experimental/BuildServerInterface.h
+++ b/Engine/Source/Developer/Zen/Public/Experimental/BuildServerInterface.h
@@ -74,6 +74,16 @@ public:
 	struct FBuildTransfer
 	{
 	public:
+		enum class EType
+		{
+			Files,
+			Oplog,
+			Count
+		};
+
+		UE_API EType GetType() const;
+		UE_API FString GetDestination() const;
+
 		UE_API FString GetDescription() const;
 		UE_API FString GetRecentOutput() const;
 		UE_API EBuildTransferStatus GetStatus() const;
diff --git a/Engine/Source/Editor/StorageServerWidgets/Internal/SBuildActivity.h b/Engine/Source/Editor/StorageServerWidgets/Internal/SBuildActivity.h
--- a/Engine/Source/Editor/StorageServerWidgets/Internal/SBuildActivity.h
+++ b/Engine/Source/Editor/StorageServerWidgets/Internal/SBuildActivity.h
@@ -49,12 +49,15 @@ class SBuildActivity : public SCompoundWidget
 private:
 	UE_API TSharedRef<ITableRow> GenerateBuildActivityRow(TSharedPtr<FBuildActivity> InItem, const TSharedRef<STableViewBase>& InOwningTable);
 	UE_API TSharedRef<SWidget> GetGridPanel();
+	void OnItemDoubleClicked(TSharedPtr<FBuildActivity, ESPMode::ThreadSafe> Item);
 	UE_API TSharedPtr<SWidget> OnGetBuildActivityContextMenuContent();
 
 	UE_API bool CanCancelBuildActivity() const;
 	UE_API void CancelBuildActivity() const;
 	UE_API bool CanRetryBuildActivity() const;
 	UE_API void RetryBuildActivity();
+	UE_API void OpenDestinationForBuildActivity();
+	UE_API void OpenDestinationForBuildActivity(TSharedPtr<FBuildActivity> Item);
 	UE_API bool CanClearBuildActivity() const;
 	UE_API void ClearBuildActivity();
 	UE_API bool CanViewLogForBuildActivity() const;
diff --git a/Engine/Source/Editor/StorageServerWidgets/Private/SBuildActivity.cpp b/Engine/Source/Editor/StorageServerWidgets/Private/SBuildActivity.cpp
--- a/Engine/Source/Editor/StorageServerWidgets/Private/SBuildActivity.cpp
+++ b/Engine/Source/Editor/StorageServerWidgets/Private/SBuildActivity.cpp
@@ -86,6 +86,7 @@ TSharedRef<SWidget> SBuildActivity::GetGridPanel()
 		SAssignNew(BuildActivityListView, SListView<FBuildActivityPtr>)
 			.ListItemsSource(&BuildActivities)
 			.OnGenerateRow(this, &SBuildActivity::GenerateBuildActivityRow)
+			.OnMouseButtonDoubleClick(this, &SBuildActivity::OnItemDoubleClicked)
 			.OnContextMenuOpening(this, &SBuildActivity::OnGetBuildActivityContextMenuContent)
 			.HeaderRow
 			(
@@ -106,6 +107,11 @@ TSharedRef<SWidget> SBuildActivity::GetGridPanel()
 	return Panel;
 }
 
+void SBuildActivity::OnItemDoubleClicked(TSharedPtr<FBuildActivity, ESPMode::ThreadSafe> Item)
+{
+	OpenDestinationForBuildActivity(Item);
+}
+
 TSharedPtr<SWidget> SBuildActivity::OnGetBuildActivityContextMenuContent()
 {
 	TArray<FBuildActivityPtr> SelectedItems = BuildActivityListView->GetSelectedItems();
@@ -128,6 +134,17 @@ TSharedPtr<SWidget> SBuildActivity::OnGetBuildActivityContextMenuContent()
 		NAME_None,
 		bRecursivelySearchable);
 
+	MenuBuilder.AddMenuEntry(
+		LOCTEXT("BuildActivity_OpenDestination", "Open destination"),
+		LOCTEXT("BuildActivity_OpenDestinationTooltip", "Open the destination directory or page for the activity"),
+		FSlateIcon(),
+		FUIAction(
+			FExecuteAction::CreateRaw(this, &SBuildActivity::OpenDestinationForBuildActivity)
+		),
+		NAME_None,
+		EUserInterfaceActionType::Button
+	);
+
 	MenuBuilder.AddMenuEntry(
 		LOCTEXT("BuildActivity_ViewLog", "View log"),
 		LOCTEXT("BuildActivity_ViewLogTooltip", "View log file for the activity"),
@@ -244,6 +261,62 @@ void SBuildActivity::RetryBuildActivity()
 	}
 }
 
+void SBuildActivity::OpenDestinationForBuildActivity()
+{
+	TArray<FBuildActivityPtr> SelectedItems = BuildActivityListView->GetSelectedItems();
+	if (SelectedItems.Num() != 1)
+	{
+		return;
+	}
+	OpenDestinationForBuildActivity(SelectedItems[0]);
+}
+
+void SBuildActivity::OpenDestinationForBuildActivity(TSharedPtr<FBuildActivity> Item)
+{
+	using namespace UE::Zen::Build;
+
+	if (TSharedPtr<FBuildServiceInstance> ServiceInstance = BuildServiceInstance.Get())
+	{
+		switch (Item->Transfer.GetType())
+		{
+		case FBuildServiceInstance::FBuildTransfer::EType::Oplog:
+			{
+				const FString Separator(TEXT("/"));
+				FString ProjectId;
+				FString OplogId;
+				if (Item->Transfer.GetDestination().Split(Separator, &ProjectId, &OplogId, ESearchCase::IgnoreCase, ESearchDir::FromEnd))
+				{
+					UE::Zen::FZenLocalServiceRunContext RunContext;
+					uint16 LocalPort = 8558;
+					if (UE::Zen::TryGetLocalServiceRunContext(RunContext))
+					{
+						if (!UE::Zen::IsLocalServiceRunning(*RunContext.GetDataPath(), &LocalPort))
+						{
+							UE::Zen::StartLocalService(RunContext);
+							UE::Zen::IsLocalServiceRunning(*RunContext.GetDataPath(), &LocalPort);
+						}
+					}
+					FPlatformProcess::LaunchURL(*FString::Printf(TEXT("http://localhost:%d/dashboard/?page=oplog&project=%s&oplog=%s"), LocalPort, *ProjectId, *OplogId), nullptr, nullptr);
+				}
+			}
+			break;
+		case FBuildServiceInstance::FBuildTransfer::EType::Files:
+			{
+				FString DestinationFolder = Item->Transfer.GetDestination();
+				if (FPaths::DirectoryExists(DestinationFolder))
+				{
+					FPlatformProcess::ExploreFolder(*DestinationFolder);
+				}
+				else
+				{
+					FPlatformProcess::ExploreFolder(*FPaths::GetPath(DestinationFolder));
+				}
+			}
+			break;
+		}
+	}
+}
+
 bool SBuildActivity::CanClearBuildActivity() const
 {
 	using namespace UE::Zen::Build;
