PR #13696: BuildStorageTool - Fix crash in case there are no streams returned

diff --git a/Engine/Source/Editor/StorageServerWidgets/Private/SBuildSelection.cpp b/Engine/Source/Editor/StorageServerWidgets/Private/SBuildSelection.cpp
--- a/Engine/Source/Editor/StorageServerWidgets/Private/SBuildSelection.cpp
+++ b/Engine/Source/Editor/StorageServerWidgets/Private/SBuildSelection.cpp
@@ -251,7 +251,7 @@ void SBuildSelection::RebuildLists()
 				}
 			}
 		}
-		FString DefaultStreamValue = Streams[DefaultStreamIndex];
+		FString DefaultStreamValue = Streams.Num() > 0 ? Streams[DefaultStreamIndex] : TEXT("");
 		ConvertToSharedPtrs(Streams, StreamList);
 		// Must search to find DefaultStreamIndex because it will have re-sorted in ConverToSharedPtrs
 		const TSharedPtr<FString>* DefaultStreamListItem = StreamList.FindByPredicate([&DefaultStreamValue](const TSharedPtr<FString>& Item)
