Storage Tool - Group on buildgroup and name if either exists before commit id, as we frequently do multiple builds of the same commit but different flavors or configurations.

diff --git a/Engine/Source/Editor/StorageServerWidgets/Private/SBuildSelection.cpp b/Engine/Source/Editor/StorageServerWidgets/Private/SBuildSelection.cpp
--- a/Engine/Source/Editor/StorageServerWidgets/Private/SBuildSelection.cpp
+++ b/Engine/Source/Editor/StorageServerWidgets/Private/SBuildSelection.cpp
@@ -441,14 +441,18 @@ void SBuildSelection::RegenerateBuildGroups(UE::Zen::Build::FListBuildsState& Li
 
 	auto MakeGroupKey = [](const FString& Namespace, const FString& CommitIdentifier, const FBuildServiceInstance::FBuildRecord& BuildRecord)
 	{
-		if (!CommitIdentifier.IsEmpty())
+		if (FCbFieldView BuildGroupField = BuildRecord.Metadata["buildgroup"]; BuildGroupField.HasValue() && !BuildGroupField.HasError())
 		{
-			return FString(WriteToString<64>(Namespace, ".", CommitIdentifier));
+			return FString(WriteToString<64>(Namespace, ".", BuildGroupField.AsString()));
 		}
 		if (FCbFieldView NameField = BuildRecord.Metadata["name"]; NameField.HasValue() && !NameField.HasError())
 		{
 			return FString(WriteToString<64>(Namespace, ".", NameField.AsString()));
 		}
+		if (!CommitIdentifier.IsEmpty())
+		{
+			return FString(WriteToString<64>(Namespace, ".", CommitIdentifier));
+		}
 		return FString(WriteToString<64>(BuildRecord.BuildId));
 	};
 
@@ -494,7 +498,12 @@ void SBuildSelection::RegenerateBuildGroups(UE::Zen::Build::FListBuildsState& Li
 				}
 
 				FString ItemName;
-				if (FCbFieldView NameView = BuildRecord.Metadata["name"]; NameView.HasValue() && !NameView.HasError())
+				FCbFieldView GroupNameView = BuildRecord.Metadata["buildgroup"];
+				if (!GroupNameView.HasValue())
+				{
+					GroupNameView = BuildRecord.Metadata["name"];
+				}
+				if (FCbFieldView NameView = GroupNameView; NameView.HasValue() && !NameView.HasError())
 				{
 					// TODO: This name manipulation needs to be removed when the metadata is more consistent.
 					ItemName = FUTF8ToTCHAR(NameView.AsString());
