Allow '.' in names used for Zen project/oplog IDs.

[RN] none

diff --git a/Engine/Source/Editor/StorageServerWidgets/Private/SBuildSelection.cpp b/Engine/Source/Editor/StorageServerWidgets/Private/SBuildSelection.cpp
--- a/Engine/Source/Editor/StorageServerWidgets/Private/SBuildSelection.cpp
+++ b/Engine/Source/Editor/StorageServerWidgets/Private/SBuildSelection.cpp
@@ -1114,7 +1114,7 @@ FString SBuildSelection::SanitizeForZenId(const FString& InString)
 	FString OutString = InString;
 	for (int32 Index = 0; Index < OutString.Len(); ++Index)
 	{
-		if (!FChar::IsIdentifier(OutString[Index]))
+		if (!FChar::IsIdentifier(OutString[Index]) && OutString[Index] != TCHAR('.'))
 		{
 			OutString[Index] = TCHAR('_');
 		}
