Build Storage Tool - Seperate out configurations into seperate items, ignore bulk builds when present and identify preflight builds by adding a suffix to that column.
These are temporary workarounds to solve for some of the more complicated builds we are uploading while we do a bigger refactoring of build storage tool to allow for more flexibility in what a build can consist of (not just multiple platforms but also flavors of those platforms)

diff --git a/Engine/Source/Editor/StorageServerWidgets/Private/SBuildSelection.cpp b/Engine/Source/Editor/StorageServerWidgets/Private/SBuildSelection.cpp
--- a/Engine/Source/Editor/StorageServerWidgets/Private/SBuildSelection.cpp
+++ b/Engine/Source/Editor/StorageServerWidgets/Private/SBuildSelection.cpp
@@ -441,9 +441,24 @@ void SBuildSelection::RegenerateBuildGroups(UE::Zen::Build::FListBuildsState& Li
 
 	auto MakeGroupKey = [](const FString& Namespace, const FString& CommitIdentifier, const FBuildServiceInstance::FBuildRecord& BuildRecord)
 	{
+		// TODO: Append the configuration to the group key if present to workaround  issues with multiple uploads with the same build group but separate configuration
+		// Should be removed once we improve the platform handling to be more flexible and thus include configurations
+		FString Configuration;
+		if (FCbFieldView ConfigurationField = BuildRecord.Metadata["Configuration"]; ConfigurationField.HasValue() && !ConfigurationField.HasError())
+		{
+			Configuration = FUTF8ToTCHAR(ConfigurationField.AsString());
+		}
+
 		if (FCbFieldView BuildGroupField = BuildRecord.Metadata["buildgroup"]; BuildGroupField.HasValue() && !BuildGroupField.HasError())
 		{
-			return FString(WriteToString<64>(Namespace, ".", BuildGroupField.AsString()));
+			if (Configuration.IsEmpty()) 
+			{
+				return FString(WriteToString<64>(Namespace, ".", BuildGroupField.AsString()));
+			}
+			else
+			{
+				return FString(WriteToString<64>(Namespace, ".", BuildGroupField.AsString(), ".", Configuration));
+			}
 		}
 		if (FCbFieldView NameField = BuildRecord.Metadata["name"]; NameField.HasValue() && !NameField.HasError())
 		{
@@ -470,6 +485,24 @@ void SBuildSelection::RegenerateBuildGroups(UE::Zen::Build::FListBuildsState& Li
 				BuildGroup = MakeShared<FBuildGroup>();
 			}
 
+			// TODO: Temporary detection of bulk builds to work around issue were a build can have multiple versions of the same platform
+			bool IsBulkBuild = false;
+			if (FCbFieldView NameView = BuildRecord.Metadata["name"]; NameView.HasValue() && !NameView.HasError())
+			{
+				FString Name;
+				Name = FUTF8ToTCHAR(NameView.AsString());
+				int32 BulkStartIndex = Name.Find(TEXT("-Bulk"));
+				if (BulkStartIndex != INDEX_NONE)
+				{
+					IsBulkBuild = true;
+				}
+			}
+
+			if (IsBulkBuild)
+			{
+				 continue;
+			}
+			
 			if (BuildGroup->DisplayName.IsEmpty())
 			{
 				FString Job;
@@ -484,6 +517,20 @@ void SBuildSelection::RegenerateBuildGroups(UE::Zen::Build::FListBuildsState& Li
 					Category = FUTF8ToTCHAR(TemplateIdField.AsString());
 				}
 
+				bool IsPreflight = false;
+				if (FCbFieldView IsPreflightField = BuildRecord.Metadata["ispreflight"]; IsPreflightField.HasValue() && !IsPreflightField.HasError())
+				{
+					FString IsPreflightString;
+					IsPreflightString = FUTF8ToTCHAR(IsPreflightField.AsString());
+					IsPreflight = IsPreflightString.ToBool();;
+				}
+
+				FString Suffix("");
+				if (IsPreflight)
+				{
+					Suffix.Append("PF");
+				}
+
 				FDateTime CreatedAt;
 				if (FCbFieldView CreatedAtField = BuildRecord.Metadata["createdAt"]; CreatedAtField.HasValue() && !CreatedAtField.HasError())
 				{
@@ -497,6 +544,12 @@ void SBuildSelection::RegenerateBuildGroups(UE::Zen::Build::FListBuildsState& Li
 					}
 				}
 
+				FString Configuration;
+				if (FCbFieldView ConfigurationField = BuildRecord.Metadata["Configuration"]; ConfigurationField.HasValue() && !ConfigurationField.HasError())
+				{
+					Configuration = FUTF8ToTCHAR(ConfigurationField.AsString());
+				}
+
 				FString ItemName;
 				FCbFieldView GroupNameView = BuildRecord.Metadata["buildgroup"];
 				if (!GroupNameView.HasValue())
@@ -526,6 +579,13 @@ void SBuildSelection::RegenerateBuildGroups(UE::Zen::Build::FListBuildsState& Li
 					{
 						ItemName.RightChopInline(Category.Len());
 					}
+
+					// TODO: Configuration appending is temporary until we have improved the platform grouping to include configuration
+					if (!Configuration.IsEmpty() && !ItemName.EndsWith(Configuration))
+					{
+						// append the configuration to the item name if present
+						ItemName = *WriteToString<64>(ItemName, "-", Configuration);
+					}
 					bool bCharRemoved = false;
 					do
 					{
@@ -548,6 +608,7 @@ void SBuildSelection::RegenerateBuildGroups(UE::Zen::Build::FListBuildsState& Li
 				BuildGroup->Namespace = BuildState.Namespace;
 				BuildGroup->DisplayName = ItemName;
 				BuildGroup->CommitIdentifier = CommitIdentifier;
+				BuildGroup->Suffix = Suffix;
 				BuildGroup->Category = Category;
 				BuildGroup->CreatedAt = CreatedAt;
 				BuildGroup->Job = Job;
