BuildStorageTool: Make the default build type prefer: Staged builds, Packaged builds, other, In that order.  Also adjust the presentation of a failed non-interactive login to make it look equivalent in the UI to not having started any login attempt.  This avoids user confusion at being presented with a failed login message when it just requires interactive login.

diff --git a/Engine/Source/Editor/StorageServerWidgets/Internal/SBuildLogin.h b/Engine/Source/Editor/StorageServerWidgets/Internal/SBuildLogin.h
--- a/Engine/Source/Editor/StorageServerWidgets/Internal/SBuildLogin.h
+++ b/Engine/Source/Editor/StorageServerWidgets/Internal/SBuildLogin.h
@@ -24,10 +24,14 @@ class SBuildLogin : public SCompoundWidget
 
 	SLATE_END_ARGS()
 
-	UE_API void Construct(const FArguments& InArgs);
+	void Construct(const FArguments& InArgs);
 
 private:
-	UE_API TSharedRef<SWidget> GetGridPanel();
+	UE::Zen::Build::FBuildServiceInstance::EConnectionState GetEffectiveConnectionState(TSharedPtr<UE::Zen::Build::FBuildServiceInstance> ServiceInstance) const;
+	TSharedRef<SWidget> GetGridPanel();
+
+	UE::Zen::Build::FBuildServiceInstance::EConnectionFailureReason LastConnectionFailureReason = UE::Zen::Build::FBuildServiceInstance::EConnectionFailureReason::None;
+	bool bTriggeredInteractiveLogin = false;
 
 	SVerticalBox::FSlot* GridSlot = nullptr;
 	TAttribute<TSharedPtr<UE::Zen::FZenServiceInstance>> ZenServiceInstance;
diff --git a/Engine/Source/Editor/StorageServerWidgets/Internal/SBuildSelection.h b/Engine/Source/Editor/StorageServerWidgets/Internal/SBuildSelection.h
--- a/Engine/Source/Editor/StorageServerWidgets/Internal/SBuildSelection.h
+++ b/Engine/Source/Editor/StorageServerWidgets/Internal/SBuildSelection.h
@@ -74,8 +74,9 @@ private:
 	{
 		EBuildType Type;
 		FText UserText;
+		int DefaultSelectionOrdering;
 	};
-	static TPair<FRegexPattern, FKnownBuildType> MakeKnownBuildTypePattern(const FStringView InPattern, EBuildType Type, FText&& InText);
+	static TPair<FRegexPattern, FKnownBuildType> MakeKnownBuildTypePattern(const FStringView InPattern, EBuildType Type, FText&& InText, int InDefaultSelectionOrdering = 0);
 	UE_API EBuildType GetSelectedBuildType() const;
 	UE_API void RebuildLists();
 	UE_API void RegenerateBuildGroups(UE::Zen::Build::FListBuildsState& ListBuildsState);
@@ -95,6 +96,7 @@ private:
 	void OnColumnSortModeChanged(const EColumnSortPriority::Type SortPriority, const FName& ColumnId, const EColumnSortMode::Type InSortMode);
 
 	FText ConvertBuildTypeToText(const FString& InBuildType);
+	int ConvertBuildTypeToDefaultSelectionOrdering(const FString& InBuildType);
 
 	UE_API void SetUserSelectedDestination(const FStringView InDestination);
 	UE_API FString GetUserSelectedDestination() const;
diff --git a/Engine/Source/Editor/StorageServerWidgets/Private/SBuildLogin.cpp b/Engine/Source/Editor/StorageServerWidgets/Private/SBuildLogin.cpp
--- a/Engine/Source/Editor/StorageServerWidgets/Private/SBuildLogin.cpp
+++ b/Engine/Source/Editor/StorageServerWidgets/Private/SBuildLogin.cpp
@@ -19,6 +19,7 @@
 
 void SBuildLogin::Construct(const FArguments& InArgs)
 {
+	using namespace UE::Zen::Build;
 	ZenServiceInstance = InArgs._ZenServiceInstance;
 	BuildServiceInstance = InArgs._BuildServiceInstance;
 
@@ -35,22 +36,44 @@ void SBuildLogin::Construct(const FArguments& InArgs)
 		]
 	];
 
-	if (TSharedPtr<UE::Zen::Build::FBuildServiceInstance> ServiceInstance = BuildServiceInstance.Get())
+	if (TSharedPtr<FBuildServiceInstance> ServiceInstance = BuildServiceInstance.Get())
 	{
-		ServiceInstance->Connect(false, [ServiceInstance]
-			(UE::Zen::Build::FBuildServiceInstance::EConnectionState ConnectionState,
-				UE::Zen::Build::FBuildServiceInstance::EConnectionFailureReason FailureReason)
+		ServiceInstance->Connect(false, [this, ServiceInstance]
+			(FBuildServiceInstance::EConnectionState ConnectionState,
+				FBuildServiceInstance::EConnectionFailureReason FailureReason)
 			{
-				if (ConnectionState == UE::Zen::Build::FBuildServiceInstance::EConnectionState::ConnectionSucceeded)
+				if (ConnectionState == FBuildServiceInstance::EConnectionState::ConnectionSucceeded)
 				{
 					ServiceInstance->RefreshNamespacesAndBuckets();
 				}
+				else
+				{
+					LastConnectionFailureReason = FailureReason;
+				}
 			});
 	}
 }
 
+UE::Zen::Build::FBuildServiceInstance::EConnectionState SBuildLogin::GetEffectiveConnectionState(TSharedPtr<UE::Zen::Build::FBuildServiceInstance> ServiceInstance) const
+{
+	using namespace UE::Zen::Build;
+	FBuildServiceInstance::EConnectionState ConnectionState = ServiceInstance->GetConnectionState();
+	if (!bTriggeredInteractiveLogin &&
+		(ConnectionState == FBuildServiceInstance::EConnectionState::ConnectionFailed) &&
+		(LastConnectionFailureReason == FBuildServiceInstance::EConnectionFailureReason::FailedToAcquireAccessToken)
+		)
+	{
+		// If we've only tried non-interactive login and it failed due to not being able to acquire a token, treat it as if we've
+		// not started the connection for UI purposes.  Showing a failure in that scenario is confusing messaging.
+		ConnectionState = FBuildServiceInstance::EConnectionState::NotStarted;
+	}
+
+	return ConnectionState;
+}
+
 TSharedRef<SWidget> SBuildLogin::GetGridPanel()
 {
+	using namespace UE::Zen::Build;
 	TSharedRef<SGridPanel> Panel = SNew(SGridPanel);
 
 	const float RowMargin = 0.0f;
@@ -73,20 +96,22 @@ TSharedRef<SWidget> SBuildLogin::GetGridPanel()
 		SNew(STextBlock)
 		.Text_Lambda([this]
 		{
-			if (TSharedPtr<UE::Zen::Build::FBuildServiceInstance> ServiceInstance = BuildServiceInstance.Get())
+			if (TSharedPtr<FBuildServiceInstance> ServiceInstance = BuildServiceInstance.Get())
 			{
-				switch (ServiceInstance->GetConnectionState())
+				FBuildServiceInstance::EConnectionState ConnectionState = GetEffectiveConnectionState(ServiceInstance);
+
+				switch (ConnectionState)
 				{
-				case UE::Zen::Build::FBuildServiceInstance::EConnectionState::NotStarted:
+				case FBuildServiceInstance::EConnectionState::NotStarted:
 					return LOCTEXT("BuildConnect_StatusValueNotConnected", "Not Connected");
-				case UE::Zen::Build::FBuildServiceInstance::EConnectionState::ConnectionInProgress:
+				case FBuildServiceInstance::EConnectionState::ConnectionInProgress:
 					return LOCTEXT("BuildConnect_StatusValueConnecting", "Connecting...");
-				case UE::Zen::Build::FBuildServiceInstance::EConnectionState::ConnectionSucceeded:
+				case FBuildServiceInstance::EConnectionState::ConnectionSucceeded:
 					{
 						return FText::Format(LOCTEXT("BuildConnect_StatusValueConnected", "Connected to {0}"),
 							FText::FromStringView(ServiceInstance->GetEffectiveDomain()));
 					}
-				case UE::Zen::Build::FBuildServiceInstance::EConnectionState::ConnectionFailed:
+				case FBuildServiceInstance::EConnectionState::ConnectionFailed:
 					return LOCTEXT("BuildConnect_StatusValueConnectionFailed", "Connection failed");
 				}
 			}
@@ -94,9 +119,9 @@ TSharedRef<SWidget> SBuildLogin::GetGridPanel()
 		})
 		.Visibility_Lambda([this]
 		{
-			if (TSharedPtr<UE::Zen::Build::FBuildServiceInstance> ServiceInstance = BuildServiceInstance.Get())
+			if (TSharedPtr<FBuildServiceInstance> ServiceInstance = BuildServiceInstance.Get())
 			{
-				if (ServiceInstance->GetConnectionState() == UE::Zen::Build::FBuildServiceInstance::EConnectionState::NotStarted)
+				if (ServiceInstance->GetConnectionState() == FBuildServiceInstance::EConnectionState::NotStarted)
 				{
 					return EVisibility::Collapsed;
 				}
@@ -111,9 +136,9 @@ TSharedRef<SWidget> SBuildLogin::GetGridPanel()
 		SNew(SHorizontalBox)
 		.Visibility_Lambda([this]
 		{
-			if (TSharedPtr<UE::Zen::Build::FBuildServiceInstance> ServiceInstance = BuildServiceInstance.Get())
+			if (TSharedPtr<FBuildServiceInstance> ServiceInstance = BuildServiceInstance.Get())
 			{
-				if (ServiceInstance->GetConnectionState() != UE::Zen::Build::FBuildServiceInstance::EConnectionState::ConnectionInProgress)
+				if (ServiceInstance->GetConnectionState() != FBuildServiceInstance::EConnectionState::ConnectionInProgress)
 				{
 					return EVisibility::Visible;
 				}
@@ -135,12 +160,14 @@ TSharedRef<SWidget> SBuildLogin::GetGridPanel()
 			.Style(FAppStyle::Get(), TEXT("NavigationHyperlink"))
 			.Text_Lambda([this]
 			{
-				if (TSharedPtr<UE::Zen::Build::FBuildServiceInstance> ServiceInstance = BuildServiceInstance.Get())
+				if (TSharedPtr<FBuildServiceInstance> ServiceInstance = BuildServiceInstance.Get())
 				{
-					switch (ServiceInstance->GetConnectionState())
+					FBuildServiceInstance::EConnectionState ConnectionState = GetEffectiveConnectionState(ServiceInstance);
+
+					switch (ConnectionState)
 					{
-					case UE::Zen::Build::FBuildServiceInstance::EConnectionState::ConnectionSucceeded:
-					case UE::Zen::Build::FBuildServiceInstance::EConnectionState::ConnectionFailed:
+					case FBuildServiceInstance::EConnectionState::ConnectionSucceeded:
+					case FBuildServiceInstance::EConnectionState::ConnectionFailed:
 						return LOCTEXT("BuildLogin_ReconnectLink", "reconnect");
 					}
 				}
@@ -148,16 +175,21 @@ TSharedRef<SWidget> SBuildLogin::GetGridPanel()
 			})
 			.OnNavigate_Lambda([this]
 			{
-				if (TSharedPtr<UE::Zen::Build::FBuildServiceInstance> ServiceInstance = BuildServiceInstance.Get())
+				if (TSharedPtr<FBuildServiceInstance> ServiceInstance = BuildServiceInstance.Get())
 				{
-					ServiceInstance->Connect(true, [ServiceInstance]
-						(UE::Zen::Build::FBuildServiceInstance::EConnectionState ConnectionState,
-							UE::Zen::Build::FBuildServiceInstance::EConnectionFailureReason FailureReason)
+					bTriggeredInteractiveLogin = true;
+					ServiceInstance->Connect(true, [this, ServiceInstance]
+						(FBuildServiceInstance::EConnectionState ConnectionState,
+							FBuildServiceInstance::EConnectionFailureReason FailureReason)
 						{
-							if (ConnectionState == UE::Zen::Build::FBuildServiceInstance::EConnectionState::ConnectionSucceeded)
+							if (ConnectionState == FBuildServiceInstance::EConnectionState::ConnectionSucceeded)
 							{
 								ServiceInstance->RefreshNamespacesAndBuckets();
 							}
+							else
+							{
+								LastConnectionFailureReason = FailureReason;
+							}
 						});
 				}
 			})
diff --git a/Engine/Source/Editor/StorageServerWidgets/Private/SBuildSelection.cpp b/Engine/Source/Editor/StorageServerWidgets/Private/SBuildSelection.cpp
--- a/Engine/Source/Editor/StorageServerWidgets/Private/SBuildSelection.cpp
+++ b/Engine/Source/Editor/StorageServerWidgets/Private/SBuildSelection.cpp
@@ -45,10 +45,10 @@ namespace UE::BuildSelection::Internal
 	}
 }
 
-TPair<FRegexPattern, SBuildSelection::FKnownBuildType> SBuildSelection::MakeKnownBuildTypePattern(const FStringView InPattern, EBuildType Type, FText&& InText)
+TPair<FRegexPattern, SBuildSelection::FKnownBuildType> SBuildSelection::MakeKnownBuildTypePattern(const FStringView InPattern, EBuildType Type, FText&& InText, int InDefaultSelectionOrdering)
 {
 	FString Pattern(InPattern);
-	return TPair<FRegexPattern, FKnownBuildType>(Pattern, FKnownBuildType{Type, InText});
+	return TPair<FRegexPattern, FKnownBuildType>(Pattern, FKnownBuildType{Type, InText, InDefaultSelectionOrdering});
 }
 
 void SBuildSelection::Construct(const FArguments& InArgs)
@@ -58,8 +58,8 @@ void SBuildSelection::Construct(const FArguments& InArgs)
 	KnownBuildTypePatterns = {
 		// These should be made more constinent and perhaps specified in INI
 		MakeKnownBuildTypePattern(TEXT(".*oplog-?(.*)"), EBuildType::Oplog, LOCTEXT("BuildSelection_BuildType_Oplog", "Cooked Data") ),
-		MakeKnownBuildTypePattern(TEXT(".*packaged-?build-?(.*)"), EBuildType::PackagedBuild, LOCTEXT("BuildSelection_BuildType_PackagedBuild", "Packaged Build")),
-		MakeKnownBuildTypePattern(TEXT(".*staged-?(.*?)-?build.*"), EBuildType::StagedBuild, LOCTEXT("BuildSelection_BuildType_StagedBuild", "Staged Build")),
+		MakeKnownBuildTypePattern(TEXT(".*packaged-?build-?(.*)"), EBuildType::PackagedBuild, LOCTEXT("BuildSelection_BuildType_PackagedBuild", "Packaged Build"), 1),
+		MakeKnownBuildTypePattern(TEXT(".*staged-?(.*?)-?build.*"), EBuildType::StagedBuild, LOCTEXT("BuildSelection_BuildType_StagedBuild", "Staged Build"), 2),
 		MakeKnownBuildTypePattern(TEXT(".*ugs-?pcb-?(.*)"), EBuildType::EditorPreCompiledBinary, LOCTEXT("BuildSelection_BuildType_PreCompiledBinary", "Editor Pre-compiled Binary")),
 		MakeKnownBuildTypePattern(TEXT(".*installed-?build-?(.*)"), EBuildType::EditorPreCompiledBinary, LOCTEXT("BuildSelection_BuildType_PreCompiledBinary", "Editor Pre-compiled Binary"))
 	};
@@ -317,7 +317,20 @@ void SBuildSelection::RebuildLists()
 			}
 		}
 		ConvertToSharedPtrs(BuildTypes, BuildTypeList);
-		ConformSelection(SelectedBuildType, BuildTypeList);
+		TArray<TSharedPtr<FString>> BuildTypeListSortedByDefaultPreference = BuildTypeList;
+		BuildTypeListSortedByDefaultPreference.StableSort([this] (const TSharedPtr<FString>& A, const TSharedPtr<FString>& B)
+		{
+			int ASelectionOrdering = ConvertBuildTypeToDefaultSelectionOrdering(*A);
+			int BSelectionOrdering = ConvertBuildTypeToDefaultSelectionOrdering(*B);
+
+			if (ASelectionOrdering != BSelectionOrdering)
+			{
+				return ASelectionOrdering >= BSelectionOrdering;
+			}
+
+			return A < B;
+		});
+		ConformSelection(SelectedBuildType, BuildTypeListSortedByDefaultPreference);
 
 		// Platform list generation
 		for (const TPair<FStringView, TArray<FStringView>>& NamespaceToBucketSegmentViews : NamespacesToBucketSegmentViews)
@@ -1011,6 +1024,20 @@ FText SBuildSelection::ConvertBuildTypeToText(const FString& InBuildType)
 	return FText::FromString(InBuildType);
 }
 
+int SBuildSelection::ConvertBuildTypeToDefaultSelectionOrdering(const FString& InBuildType)
+{
+	for (const TPair<FRegexPattern, FKnownBuildType>& KnownBuildTypeItem : KnownBuildTypePatterns)
+	{
+		FRegexMatcher Matcher(KnownBuildTypeItem.Key, InBuildType);
+		if (Matcher.FindNext())
+		{
+			return KnownBuildTypeItem.Value.DefaultSelectionOrdering;
+		}
+	}
+
+	return 0;
+}
+
 void SBuildSelection::SetUserSelectedDestination(const FStringView InDestination)
 {
 	EBuildType BuildType = GetSelectedBuildType();
