BuildStorageTool: Expand text search to cover all fields visible in the build group list instead of being just for commits.  Add explicit filter for including preflights.  Change label of Required Platforms to just Platforms.  Fix some bugs identified post review.

diff --git a/Engine/Source/Editor/StorageServerWidgets/Internal/SBuildSelection.h b/Engine/Source/Editor/StorageServerWidgets/Internal/SBuildSelection.h
--- a/Engine/Source/Editor/StorageServerWidgets/Internal/SBuildSelection.h
+++ b/Engine/Source/Editor/StorageServerWidgets/Internal/SBuildSelection.h
@@ -55,6 +55,7 @@ class SBuildSelection : public SCompoundWidget
 		FString Category;
 		FString Job;
 		TMap<FString, UE::Zen::Build::FBuildServiceInstance::FBuildRecord> PerPlatformBuilds;
+		bool bIsPreflight;
 	};
 
 	UE_API void Construct(const FArguments& InArgs);
@@ -77,35 +78,37 @@ private:
 		int DefaultSelectionOrdering;
 	};
 	static TPair<FRegexPattern, FKnownBuildType> MakeKnownBuildTypePattern(const FStringView InPattern, EBuildType Type, FText&& InText, int InDefaultSelectionOrdering = 0);
-	UE_API EBuildType GetSelectedBuildType() const;
-	UE_API void RebuildLists();
-	UE_API void RegenerateBuildGroups(UE::Zen::Build::FListBuildsState& ListBuildsState);
-	UE_API void ConditionalSortBuildGroups();
-	UE_API void RegenerateActivePlatformFilters();
-	UE_API void ValidateBuildGroupSelection();
-	UE_API TSharedRef<SWidget> OnGenerateTextBlockFromString(TSharedPtr<FString> Item);
-	UE_API TSharedRef<SWidget> OnGenerateBuildTypeTextBlockFromString(TSharedPtr<FString> Item);
-	UE_API bool BuildGroupIsSelectableOrNavigable(TSharedPtr<FBuildGroup> InItem) const;
-	UE_API TSharedRef<ITableRow> GenerateBuildGroupRow(TSharedPtr<FBuildGroup> InItem, const TSharedRef<STableViewBase>& InOwningTable);
-	UE_API TSharedPtr<SWidget> OnGetBuildGroupContextMenuContent() const;
-	UE_API void OnOpenDestinationDirectoryClicked();
-	UE_API TSharedRef<SWidget> GetBuildDestinationPanel();
-	UE_API TSharedRef<SWidget> GetGridPanel();
-	UE_API void BuildGroupSelectionChanged(TSharedPtr<FBuildGroup> Item, ESelectInfo::Type SelectInfo);
+	EBuildType GetSelectedBuildType() const;
+	void RebuildLists();
+	void RegenerateBuildGroups(UE::Zen::Build::FListBuildsState& ListBuildsState);
+	void ConditionalSortBuildGroups();
+	void RegenerateActivePlatformFilters();
+	void ValidateBuildGroupSelection();
+	TSharedRef<SWidget> OnGenerateTextBlockFromString(TSharedPtr<FString> Item);
+	TSharedRef<SWidget> OnGenerateBuildTypeTextBlockFromString(TSharedPtr<FString> Item);
+	bool DoesPassAllViewFilters(TSharedPtr<FBuildGroup> InItem) const;
+	bool BuildGroupIsSelectableOrNavigable(TSharedPtr<FBuildGroup> InItem) const;
+	TSharedRef<ITableRow> GenerateBuildGroupRow(TSharedPtr<FBuildGroup> InItem, const TSharedRef<STableViewBase>& InOwningTable);
+	TSharedPtr<SWidget> OnGetBuildGroupContextMenuContent() const;
+	void OnOpenDestinationDirectoryClicked();
+	TSharedRef<SWidget> GetBuildDestinationPanel();
+	TSharedRef<SWidget> GetGridPanel();
+	void BuildGroupSelectionChanged(TSharedPtr<FBuildGroup> Item, ESelectInfo::Type SelectInfo);
 	EColumnSortMode::Type GetColumnSortMode(const FName ColumnId) const;
 	void OnColumnSortModeChanged(const EColumnSortPriority::Type SortPriority, const FName& ColumnId, const EColumnSortMode::Type InSortMode);
 
 	FText ConvertBuildTypeToText(const FString& InBuildType);
 	int ConvertBuildTypeToDefaultSelectionOrdering(const FString& InBuildType);
 
-	UE_API void SetUserSelectedDestination(const FStringView InDestination);
-	UE_API FString GetUserSelectedDestination() const;
-	UE_API FString GetDefaultDestination() const;
-	UE_API FString GetEffectiveDestination() const;
+	void SetUserSelectedDestination(const FStringView InDestination);
+	FString GetUserSelectedDestination() const;
+	FString GetDefaultDestination() const;
+	FString GetEffectiveDestination() const;
 
 	void SetSelectedStream(const TSharedPtr<FString> InSelectedStream);
 	void SetSelectedProject(const TSharedPtr<FString> InSelectedProject);
 	void SetSelectedBuildType(const TSharedPtr<FString> InBuildType);
+	void SetIncludePreflights(const bool bInIncludePreflights);
 
 	bool WriteSetting(FStringView InSectionName, FStringView InKeyName, FStringView InValue);
 	bool WriteSetting(FStringView InKeyName, FStringView InValue);
@@ -114,13 +117,13 @@ private:
 
 	void SetUserSelectedProjectDictionaryRoot(FStringView InRoot);
 
-	static UE_API FString SanitizeForPath(const FString& InString);
-	static UE_API FString SanitizeForZenId(const FString& InString);
-	static UE_API FString SanitizeBucketSegment(const FString& InString);
+	static FString SanitizeForPath(const FString& InString);
+	static FString SanitizeForZenId(const FString& InString);
+	static FString SanitizeBucketSegment(const FString& InString);
 
 	static const TCHAR* LexToString(EBuildType BuildType);
 
-	UE_API FReply ExploreDestination_OnClicked();
+	FReply ExploreDestination_OnClicked();
 
 	static const FString EmptyString;
 
@@ -140,7 +143,7 @@ private:
 	TSharedPtr<SComboBox<TSharedPtr<FString>>> ProjectWidget;
 	TSharedPtr<SComboBox<TSharedPtr<FString>>> StreamWidget;
 	TSharedPtr<SComboBox<TSharedPtr<FString>>> BuildTypeWidget;
-	TSharedPtr<SMultiSelectComboBox> RequiredPlatformsWidget;
+	TSharedPtr<SMultiSelectComboBox> PlatformFiltersWidget;
 
 	TSharedPtr<SGridPanel> SelectedGroupPlatformGrid;
 	TArray<FString> SelectedGroupSelectedPlatforms;
@@ -159,7 +162,8 @@ private:
 	TSharedPtr<FString> SelectedStream;
 	TSharedPtr<FString> SelectedProject;
 	TSharedPtr<FString> SelectedBuildType;
-	FString SelectedCommitFilter;
+	FString SelectedTextFilter;
+	bool bIncludePreflights = true;
 
 	TArray<TPair<FRegexPattern, FKnownBuildType>> KnownBuildTypePatterns;
 
diff --git a/Engine/Source/Editor/StorageServerWidgets/Private/SBuildLogin.cpp b/Engine/Source/Editor/StorageServerWidgets/Private/SBuildLogin.cpp
--- a/Engine/Source/Editor/StorageServerWidgets/Private/SBuildLogin.cpp
+++ b/Engine/Source/Editor/StorageServerWidgets/Private/SBuildLogin.cpp
@@ -38,7 +38,7 @@ void SBuildLogin::Construct(const FArguments& InArgs)
 
 	if (TSharedPtr<FBuildServiceInstance> ServiceInstance = BuildServiceInstance.Get())
 	{
-		ServiceInstance->Connect(false, [this, ServiceInstance]
+		ServiceInstance->Connect(false, [this, SharedThis=AsShared(), ServiceInstance]
 			(FBuildServiceInstance::EConnectionState ConnectionState,
 				FBuildServiceInstance::EConnectionFailureReason FailureReason)
 			{
@@ -178,7 +178,7 @@ TSharedRef<SWidget> SBuildLogin::GetGridPanel()
 				if (TSharedPtr<FBuildServiceInstance> ServiceInstance = BuildServiceInstance.Get())
 				{
 					bTriggeredInteractiveLogin = true;
-					ServiceInstance->Connect(true, [this, ServiceInstance]
+					ServiceInstance->Connect(true, [this, SharedThis=AsShared(), ServiceInstance]
 						(FBuildServiceInstance::EConnectionState ConnectionState,
 							FBuildServiceInstance::EConnectionFailureReason FailureReason)
 						{
diff --git a/Engine/Source/Editor/StorageServerWidgets/Private/SBuildSelection.cpp b/Engine/Source/Editor/StorageServerWidgets/Private/SBuildSelection.cpp
--- a/Engine/Source/Editor/StorageServerWidgets/Private/SBuildSelection.cpp
+++ b/Engine/Source/Editor/StorageServerWidgets/Private/SBuildSelection.cpp
@@ -7,6 +7,7 @@
 #include "HAL/PlatformApplicationMisc.h"
 #include "HAL/PlatformFileManager.h"
 #include "Internationalization/FastDecimalFormat.h"
+#include "Internationalization/TextChronoFormatter.h"
 #include "Math/BasicMathExpressionEvaluator.h"
 #include "Math/UnitConversion.h"
 #include "Misc/App.h"
@@ -118,6 +119,12 @@ void SBuildSelection::Construct(const FArguments& InArgs)
 	{
 		ReadSetting(TEXT("Branch"), BranchName);
 	}
+	
+	FString IncludePreflightsString;
+	if (ReadSetting(TEXT("IncludePreflights"), IncludePreflightsString))
+	{
+		bIncludePreflights = FCString::ToBool(*IncludePreflightsString);
+	}
 
 	for (bool& bAppendBuildNameToDestination : bAppendBuildNameToDestinations)
 	{
@@ -328,7 +335,7 @@ void SBuildSelection::RebuildLists()
 				return ASelectionOrdering >= BSelectionOrdering;
 			}
 
-			return A < B;
+			return *A < *B;
 		});
 		ConformSelection(SelectedBuildType, BuildTypeListSortedByDefaultPreference);
 
@@ -452,7 +459,9 @@ void SBuildSelection::RegenerateBuildGroups(UE::Zen::Build::FListBuildsState& Li
 
 	BuildGroups.Empty();
 
-	auto MakeGroupKey = [](const FString& Namespace, const FString& CommitIdentifier, const FBuildServiceInstance::FBuildRecord& BuildRecord)
+	const bool bDoConfigurationSplitting = true;
+
+	auto MakeGroupKey = [bDoConfigurationSplitting](const FString& Namespace, const FString& CommitIdentifier, const FBuildServiceInstance::FBuildRecord& BuildRecord)
 	{
 		// TODO: Append the configuration to the group key if present to workaround  issues with multiple uploads with the same build group but separate configuration
 		// Should be removed once we improve the platform handling to be more flexible and thus include configurations
@@ -464,7 +473,7 @@ void SBuildSelection::RegenerateBuildGroups(UE::Zen::Build::FListBuildsState& Li
 
 		if (FCbFieldView BuildGroupField = BuildRecord.Metadata["buildgroup"]; BuildGroupField.HasValue() && !BuildGroupField.HasError())
 		{
-			if (Configuration.IsEmpty()) 
+			if (Configuration.IsEmpty() || !bDoConfigurationSplitting) 
 			{
 				return FString(WriteToString<64>(Namespace, ".", BuildGroupField.AsString()));
 			}
@@ -530,16 +539,16 @@ void SBuildSelection::RegenerateBuildGroups(UE::Zen::Build::FListBuildsState& Li
 					Category = FUTF8ToTCHAR(TemplateIdField.AsString());
 				}
 
-				bool IsPreflight = false;
+				bool bIsPreflight = false;
 				if (FCbFieldView IsPreflightField = BuildRecord.Metadata["ispreflight"]; IsPreflightField.HasValue() && !IsPreflightField.HasError())
 				{
 					FString IsPreflightString;
 					IsPreflightString = FUTF8ToTCHAR(IsPreflightField.AsString());
-					IsPreflight = IsPreflightString.ToBool();;
+					bIsPreflight = IsPreflightString.ToBool();
 				}
 
 				FString Suffix("");
-				if (IsPreflight)
+				if (bIsPreflight)
 				{
 					Suffix.Append("PF");
 				}
@@ -594,7 +603,7 @@ void SBuildSelection::RegenerateBuildGroups(UE::Zen::Build::FListBuildsState& Li
 					}
 
 					// TODO: Configuration appending is temporary until we have improved the platform grouping to include configuration
-					if (!Configuration.IsEmpty() && !ItemName.EndsWith(Configuration))
+					if (bDoConfigurationSplitting && !Configuration.IsEmpty() && !ItemName.EndsWith(Configuration))
 					{
 						// append the configuration to the item name if present
 						ItemName = *WriteToString<64>(ItemName, "-", Configuration);
@@ -625,6 +634,7 @@ void SBuildSelection::RegenerateBuildGroups(UE::Zen::Build::FListBuildsState& Li
 				BuildGroup->Category = Category;
 				BuildGroup->CreatedAt = CreatedAt;
 				BuildGroup->Job = Job;
+				BuildGroup->bIsPreflight = bIsPreflight;
 			}
 
 			BuildGroup->PerPlatformBuilds.FindOrAdd(BuildState.Platform, MoveTemp(BuildRecord));
@@ -718,9 +728,9 @@ void SBuildSelection::RegenerateActivePlatformFilters()
 	ActivePlatformFilters.Empty();
 	for (TSharedPtr<FString> Platform : PlatformList)
 	{
-		if (Platform && RequiredPlatformsWidget)
+		if (Platform && PlatformFiltersWidget)
 		{
-			if (RequiredPlatformsWidget->IsChecked(*Platform))
+			if (PlatformFiltersWidget->IsChecked(*Platform))
 			{
 				ActivePlatformFilters.Add(*Platform);
 				SelectedGroupSelectedPlatforms.AddUnique(*Platform);
@@ -762,7 +772,7 @@ TSharedRef<SWidget> SBuildSelection::OnGenerateBuildTypeTextBlockFromString(TSha
 		.Text(Item ? ConvertBuildTypeToText(*Item) : FText::GetEmpty());
 }
 
-bool SBuildSelection::BuildGroupIsSelectableOrNavigable(FBuildSelectionBuildGroupPtr InItem) const
+bool SBuildSelection::DoesPassAllViewFilters(TSharedPtr<FBuildGroup> InItem) const
 {
 	if (!InItem)
 	{
@@ -777,11 +787,37 @@ bool SBuildSelection::BuildGroupIsSelectableOrNavigable(FBuildSelectionBuildGrou
 		}
 	}
 
-	if (!SelectedCommitFilter.IsEmpty() && !InItem->CommitIdentifier.Contains(SelectedCommitFilter))
+	if (!bIncludePreflights && InItem->bIsPreflight)
 	{
 		return false;
 	}
-	return true;
+
+	if (SelectedTextFilter.IsEmpty())
+	{
+		return true;
+	}
+
+	if (InItem->Namespace.Contains(SelectedTextFilter) ||
+		InItem->DisplayName.Contains(SelectedTextFilter) ||
+		InItem->CommitIdentifier.Contains(SelectedTextFilter) ||
+		InItem->Suffix.Contains(SelectedTextFilter) ||
+		InItem->Category.Contains(SelectedTextFilter) ||
+		InItem->Job.Contains(SelectedTextFilter))
+	{
+		return true;
+	}
+
+	if (FText::AsDateTime(InItem->CreatedAt, EDateTimeStyle::Short).ToString().Contains(SelectedTextFilter))
+	{
+		return true;
+	}
+
+	return false;
+}
+
+bool SBuildSelection::BuildGroupIsSelectableOrNavigable(FBuildSelectionBuildGroupPtr InItem) const
+{
+	return DoesPassAllViewFilters(InItem);
 }
 
 TSharedRef<ITableRow> SBuildSelection::GenerateBuildGroupRow(FBuildSelectionBuildGroupPtr InItem, const TSharedRef<STableViewBase>& InOwningTable)
@@ -789,19 +825,11 @@ TSharedRef<ITableRow> SBuildSelection::GenerateBuildGroupRow(FBuildSelectionBuil
 	return SNew(SBuildGroupTableRow, InOwningTable, InItem)
 		.Visibility_Lambda([this, InItem]()
 		{
-			for (const FString& ActivePlatformFilter : ActivePlatformFilters)
+			if (DoesPassAllViewFilters(InItem))
 			{
-				if (!InItem->PerPlatformBuilds.Contains(ActivePlatformFilter))
-				{
-					return EVisibility::Collapsed;
-				}
-			}
-
-			if (!SelectedCommitFilter.IsEmpty() && !InItem->CommitIdentifier.Contains(SelectedCommitFilter))
-			{
-				return EVisibility::Collapsed;
+				return EVisibility::Visible;
 			}
-			return EVisibility::Visible;
+			return EVisibility::Collapsed;
 		});
 }
 
@@ -1146,6 +1174,12 @@ void SBuildSelection::SetSelectedBuildType(const TSharedPtr<FString> InBuildType
 	WriteSetting(TEXT("BuildType"), *SelectedBuildType);
 }
 
+void SBuildSelection::SetIncludePreflights(const bool bInIncludePreflights)
+{
+	bIncludePreflights = bInIncludePreflights;
+	WriteSetting(TEXT("IncludePreflights"), bIncludePreflights ? TEXT("True") : TEXT("False"));
+}
+
 bool SBuildSelection::WriteSetting(const FStringView InSectionName, const FStringView InKeyName, const FStringView InValue)
 {
 	static const FString StoreId = TEXT("Epic Games");
@@ -1639,7 +1673,7 @@ TSharedRef<SWidget> SBuildSelection::GetGridPanel()
 			]
 		]
 
-		// Commit
+		// Text Filter
 		+SHorizontalBox::Slot()
 		.AutoWidth()
 		.FillWidth(1.0f)
@@ -1656,7 +1690,7 @@ TSharedRef<SWidget> SBuildSelection::GetGridPanel()
 				SNew(STextBlock)
 				.ColorAndOpacity(TitleColor)
 				.Font(TitleFont)
-				.Text(LOCTEXT("BuildSelection_Commit", "Commit"))
+				.Text(LOCTEXT("BuildSelection_TextFilter", "Text Filter"))
 			]
 
 			+ SVerticalBox::Slot()
@@ -1667,21 +1701,17 @@ TSharedRef<SWidget> SBuildSelection::GetGridPanel()
 			[
 				SNew(SEditableTextBox)
 				.OverflowPolicy(ETextOverflowPolicy::MiddleEllipsis)
-				.HintText_Lambda([this]()
-				{
-					return LOCTEXT("BuildSelection_AnyCommit", "Any");
-				})
 				.Text_Lambda([this]()
 				{
-					return FText::FromString(SelectedCommitFilter);
+					return FText::FromString(SelectedTextFilter);
 				})
 				.OnTextChanged_Lambda([this](const FText& Text)
 				{
-					SelectedCommitFilter = Text.ToString();
+					SelectedTextFilter = Text.ToString();
 				})
 				.OnTextCommitted_Lambda([this](const FText& Text, const ETextCommit::Type CommitType)
 				{
-					SelectedCommitFilter = Text.ToString();
+					SelectedTextFilter = Text.ToString();
 				})
 			]
 		]
@@ -1703,7 +1733,7 @@ TSharedRef<SWidget> SBuildSelection::GetGridPanel()
 				SNew(STextBlock)
 				.ColorAndOpacity(TitleColor)
 				.Font(TitleFont)
-				.Text(LOCTEXT("BuildSelection_RequiredPlatforms", "Required Platforms"))
+				.Text(LOCTEXT("BuildSelection_PlatformFilters", "Platforms"))
 			]
 
 			+ SVerticalBox::Slot()
@@ -1712,7 +1742,7 @@ TSharedRef<SWidget> SBuildSelection::GetGridPanel()
 			.HAlign(HAlign_Fill)
 			.VAlign(VAlign_Top)
 			[
-				SAssignNew(RequiredPlatformsWidget, SMultiSelectComboBox)
+				SAssignNew(PlatformFiltersWidget, SMultiSelectComboBox)
 				.SelectValues(&PlatformList)
 				.OnCheckedValuesChanged_Lambda([this]()
 				{
@@ -1721,6 +1751,50 @@ TSharedRef<SWidget> SBuildSelection::GetGridPanel()
 				})
 			]
 		]
+
+		// Include Preflights
+		+SHorizontalBox::Slot()
+		.AutoWidth()
+		.HAlign(HAlign_Fill)
+		.VAlign(VAlign_Top)
+		[
+			SNew(SVerticalBox)
+
+			+ SVerticalBox::Slot()
+			.AutoHeight()
+			.Padding(FMargin(ColumnMargin, RowMargin))
+			.VAlign(VAlign_Top)
+			[
+				SNew(STextBlock)
+				.ColorAndOpacity(TitleColor)
+				.Font(TitleFont)
+				.Text(LOCTEXT("BuildSelection_ShowPreflights", "Preflights"))
+			]
+
+			+ SVerticalBox::Slot()
+			.AutoHeight()
+			.Padding(FMargin(ColumnMargin, RowMargin))
+			.HAlign(HAlign_Fill)
+			.VAlign(VAlign_Top)
+			[
+				SNew(SCheckBox)
+				.IsChecked_Lambda([this]
+				{
+					return bIncludePreflights ? ECheckBoxState::Checked : ECheckBoxState::Unchecked;
+				})
+				.OnCheckStateChanged_Lambda([this](ECheckBoxState InNewState)
+				{
+					if (InNewState == ECheckBoxState::Checked)
+					{
+						SetIncludePreflights(true);
+					}
+					else if (InNewState == ECheckBoxState::Unchecked)
+					{
+						SetIncludePreflights(false);
+					}
+				})
+			]
+		]
 	];
 
 	Panel->AddSlot()
@@ -1802,27 +1876,11 @@ TSharedRef<SWidget> SBuildSelection::GetGridPanel()
 					return LOCTEXT("BuildSelection_ResultLoading", "Loading...");
 				}
 				int32 VisibleItemCount = 0;
-				if (ActivePlatformFilters.IsEmpty())
+				for (FBuildSelectionBuildGroupPtr BuildGroup : BuildGroups)
 				{
-					VisibleItemCount = BuildGroups.Num();
-				}
-				else
-				{
-					for (FBuildSelectionBuildGroupPtr BuildGroup : BuildGroups)
+					if (DoesPassAllViewFilters(BuildGroup))
 					{
-						bool bHasAllRequiredPlatforms = true;
-						for (const FString& ActivePlatformFilter : ActivePlatformFilters)
-						{
-							if (!BuildGroup->PerPlatformBuilds.Contains(ActivePlatformFilter))
-							{
-								bHasAllRequiredPlatforms = false;
-								break;
-							}
-						}
-						if (bHasAllRequiredPlatforms)
-						{
-							VisibleItemCount++;
-						}
+						VisibleItemCount++;
 					}
 				}
 				return FText::Format(LOCTEXT("BuildSelection_ResultDescription", "{0} {0}|plural(one=item,other=items)"), FText::AsNumber(VisibleItemCount));
