BuildStorageTool: Add buttons for opening destination and canceling to the activity rows.

diff --git a/Engine/Source/Editor/StorageServerWidgets/Internal/SBuildActivity.h b/Engine/Source/Editor/StorageServerWidgets/Internal/SBuildActivity.h
--- a/Engine/Source/Editor/StorageServerWidgets/Internal/SBuildActivity.h
+++ b/Engine/Source/Editor/StorageServerWidgets/Internal/SBuildActivity.h
@@ -16,6 +16,8 @@
 
 #define UE_API STORAGESERVERWIDGETS_API
 
+class SBuildActivityTableRow;
+
 class SBuildActivity : public SCompoundWidget
 {
 	SLATE_BEGIN_ARGS(SBuildActivity)
@@ -47,22 +49,24 @@ class SBuildActivity : public SCompoundWidget
 	}
 
 private:
-	UE_API TSharedRef<ITableRow> GenerateBuildActivityRow(TSharedPtr<FBuildActivity> InItem, const TSharedRef<STableViewBase>& InOwningTable);
-	UE_API TSharedRef<SWidget> GetGridPanel();
+	friend class SBuildActivityTableRow;
+
+	TSharedRef<ITableRow> GenerateBuildActivityRow(TSharedPtr<FBuildActivity> InItem, const TSharedRef<STableViewBase>& InOwningTable);
+	TSharedRef<SWidget> GetGridPanel();
 	void OnItemDoubleClicked(TSharedPtr<FBuildActivity, ESPMode::ThreadSafe> Item);
-	UE_API TSharedPtr<SWidget> OnGetBuildActivityContextMenuContent();
-
-	UE_API bool CanCancelBuildActivity() const;
-	UE_API void CancelBuildActivity() const;
-	UE_API bool CanRetryBuildActivity() const;
-	UE_API void RetryBuildActivity();
-	UE_API void OpenDestinationForBuildActivity();
-	UE_API void OpenDestinationForBuildActivity(TSharedPtr<FBuildActivity> Item);
-	UE_API bool CanClearBuildActivity() const;
-	UE_API void ClearBuildActivity();
-	UE_API bool CanViewLogForBuildActivity() const;
-	UE_API void ViewLogForBuildActivity() const;
-	UE_API void ClearAllCompleted();
+	TSharedPtr<SWidget> OnGetBuildActivityContextMenuContent();
+
+	bool CanCancelBuildActivity() const;
+	void CancelBuildActivity() const;
+	bool CanRetryBuildActivity() const;
+	void RetryBuildActivity();
+	void OpenDestinationForBuildActivity();
+	static void OpenDestinationForBuildActivity(TSharedPtr<UE::Zen::Build::FBuildServiceInstance> ServiceInstance, TSharedPtr<FBuildActivity> Item);
+	bool CanClearBuildActivity() const;
+	void ClearBuildActivity();
+	bool CanViewLogForBuildActivity() const;
+	void ViewLogForBuildActivity() const;
+	void ClearAllCompleted();
 
 	TSharedPtr<SListView<TSharedPtr<SBuildActivity::FBuildActivity, ESPMode::ThreadSafe>>> BuildActivityListView;
 	TArray<TSharedPtr<SBuildActivity::FBuildActivity, ESPMode::ThreadSafe>> BuildActivities;
diff --git a/Engine/Source/Editor/StorageServerWidgets/Private/SBuildActivity.cpp b/Engine/Source/Editor/StorageServerWidgets/Private/SBuildActivity.cpp
--- a/Engine/Source/Editor/StorageServerWidgets/Private/SBuildActivity.cpp
+++ b/Engine/Source/Editor/StorageServerWidgets/Private/SBuildActivity.cpp
@@ -27,6 +27,7 @@ namespace UE::BuildActivity::Internal
 		const FName ColName = TEXT("Name");
 		const FName ColPlatform = TEXT("Platform");
 		const FName ColDescription = TEXT("Description");
+		const FName ColActions = TEXT("Actions");
 	}
 }
 
@@ -99,7 +100,9 @@ TSharedRef<SWidget> SBuildActivity::GetGridPanel()
 				+ SHeaderRow::Column(FBuildActivityIds::ColPlatform).DefaultLabel(LOCTEXT("BuildActivity_ActivityPlatform", "Platform"))
 					.FillWidth(0.1f).HAlignCell(HAlign_Left).HAlignHeader(HAlign_Left).VAlignCell(VAlign_Center)
 				+ SHeaderRow::Column(FBuildActivityIds::ColDescription).DefaultLabel(LOCTEXT("BuildActivity_ActivityDescription", "Description"))
-					.FillWidth(0.6f).HAlignCell(HAlign_Fill).HAlignHeader(HAlign_Center).VAlignCell(VAlign_Center)
+					.FillWidth(0.54f).HAlignCell(HAlign_Fill).HAlignHeader(HAlign_Center).VAlignCell(VAlign_Center)
+				+ SHeaderRow::Column(FBuildActivityIds::ColActions).DefaultLabel(LOCTEXT("BuildActivity_ActivityActions", "Actions"))
+					.FillWidth(0.06f).HAlignCell(HAlign_Fill).HAlignHeader(HAlign_Center).VAlignCell(VAlign_Center)
 			)
 	];
 
@@ -109,7 +112,7 @@ TSharedRef<SWidget> SBuildActivity::GetGridPanel()
 
 void SBuildActivity::OnItemDoubleClicked(TSharedPtr<FBuildActivity, ESPMode::ThreadSafe> Item)
 {
-	OpenDestinationForBuildActivity(Item);
+	OpenDestinationForBuildActivity(BuildServiceInstance.Get(), Item);
 }
 
 TSharedPtr<SWidget> SBuildActivity::OnGetBuildActivityContextMenuContent()
@@ -268,14 +271,14 @@ void SBuildActivity::OpenDestinationForBuildActivity()
 	{
 		return;
 	}
-	OpenDestinationForBuildActivity(SelectedItems[0]);
+	OpenDestinationForBuildActivity(BuildServiceInstance.Get(), SelectedItems[0]);
 }
 
-void SBuildActivity::OpenDestinationForBuildActivity(TSharedPtr<FBuildActivity> Item)
+void SBuildActivity::OpenDestinationForBuildActivity(TSharedPtr<UE::Zen::Build::FBuildServiceInstance> ServiceInstance, TSharedPtr<FBuildActivity> Item)
 {
 	using namespace UE::Zen::Build;
 
-	if (TSharedPtr<FBuildServiceInstance> ServiceInstance = BuildServiceInstance.Get())
+	if (ServiceInstance)
 	{
 		switch (Item->Transfer.GetType())
 		{
@@ -524,6 +527,58 @@ SBuildActivityTableRow::GenerateWidgetForColumn(const FName& ColumnName)
 						})
 					];
 	}
+	else if (ColumnName == FBuildActivityIds::ColActions)
+	{
+		return SNew(SHorizontalBox)
+				+ SHorizontalBox::Slot()
+				.AutoWidth()
+				.Padding(1.f,1.f)
+				[
+					SNew(SButton)
+					.ButtonStyle(&FAppStyle::Get().GetWidgetStyle<FButtonStyle>("SimpleButton"))
+					.ToolTipText(LOCTEXT("BuildActivity_OpenDestinationButtonTooltip", "Open the destination directory or page for the activity"))
+					.OnClicked_Lambda([this]
+					{
+						SBuildActivity::OpenDestinationForBuildActivity(BuildServiceInstance, BuildActivity);
+						return FReply::Handled();
+					})
+					[
+						SNew(SImage)
+						.Image_Lambda([this]
+						{
+							return BuildActivity->Transfer.GetType() == FBuildServiceInstance::FBuildTransfer::EType::Oplog ? FAppStyle::Get().GetBrush("Zen.BrowserView") : FAppStyle::Get().GetBrush("Zen.FolderView");
+						})
+						.ColorAndOpacity(FSlateColor::UseForeground())
+					]
+				]
+				+ SHorizontalBox::Slot()
+				.AutoWidth()
+				.Padding(1.f,1.f)
+				[
+					SNew(SButton)
+					.ButtonStyle(&FAppStyle::Get().GetWidgetStyle<FButtonStyle>("SimpleButton"))
+					.ToolTipText(LOCTEXT("BuildActivity_CancelButtonTooltip", "Cancel the pending or in progress activity"))
+					.IsEnabled_Lambda([this]
+					{
+						FBuildServiceInstance::EBuildTransferStatus Status = BuildActivity->Transfer.GetStatus();
+						if (Status == FBuildServiceInstance::EBuildTransferStatus::Active || Status == FBuildServiceInstance::EBuildTransferStatus::Queued)
+						{
+							return true;
+						}
+						return false;
+					})
+					.OnClicked_Lambda([this]
+					{
+						BuildActivity->Transfer.RequestCancel();
+						return FReply::Handled();
+					})
+					[
+						SNew(SImage)
+						.Image(FAppStyle::Get().GetBrush("Zen.Cancel"))
+						.ColorAndOpacity(FSlateColor::UseForeground())
+					]
+				];
+	}
 
 	return SNullWidget::NullWidget;
 }
diff --git a/Engine/Source/Programs/BuildStorageTool/Private/BuildStorageToolStyle.cpp b/Engine/Source/Programs/BuildStorageTool/Private/BuildStorageToolStyle.cpp
--- a/Engine/Source/Programs/BuildStorageTool/Private/BuildStorageToolStyle.cpp
+++ b/Engine/Source/Programs/BuildStorageTool/Private/BuildStorageToolStyle.cpp
@@ -67,6 +67,8 @@ TSharedRef< FSlateStyleSet > FBuildStorageToolStyle::Create()
 	Style.Set("Zen.FolderView", new CORE_IMAGE_BRUSH_SVG("Starship/Common/folder-open", Icon12x12));
 	Style.Set("Zen.BrowserView", new EDITOR_IMAGE_BRUSH_SVG("Starship/Common/WebBrowser", Icon12x12));
 
+	Style.Set("Zen.Cancel", new CORE_IMAGE_BRUSH_SVG("Starship/Common/reject", Icon12x12));
+
 	const FTextBlockStyle DefaultText = FTextBlockStyle()
 		.SetFont(DEFAULT_FONT("Bold", 10))
 		.SetColorAndOpacity(FSlateColor::UseForeground())
