BuildStorageTool: Only do build configuration splitting for packaged builds.

diff --git a/Engine/Source/Editor/StorageServerWidgets/Private/SBuildSelection.cpp b/Engine/Source/Editor/StorageServerWidgets/Private/SBuildSelection.cpp
--- a/Engine/Source/Editor/StorageServerWidgets/Private/SBuildSelection.cpp
+++ b/Engine/Source/Editor/StorageServerWidgets/Private/SBuildSelection.cpp
@@ -459,7 +459,7 @@ void SBuildSelection::RegenerateBuildGroups(UE::Zen::Build::FListBuildsState& Li
 
 	BuildGroups.Empty();
 
-	const bool bDoConfigurationSplitting = true;
+	const bool bDoConfigurationSplitting = GetSelectedBuildType() == EBuildType::PackagedBuild;
 
 	auto MakeGroupKey = [bDoConfigurationSplitting](const FString& Namespace, const FString& CommitIdentifier, const FBuildServiceInstance::FBuildRecord& BuildRecord)
 	{
