BuildStorageTool: Fix the overlap and layout of the "available platforms" panel shown for each build group.

diff --git a/Engine/Source/Editor/StorageServerWidgets/Private/SBuildSelection.cpp b/Engine/Source/Editor/StorageServerWidgets/Private/SBuildSelection.cpp
--- a/Engine/Source/Editor/StorageServerWidgets/Private/SBuildSelection.cpp
+++ b/Engine/Source/Editor/StorageServerWidgets/Private/SBuildSelection.cpp
@@ -10,6 +10,7 @@
 #include "Internationalization/TextChronoFormatter.h"
 #include "Math/BasicMathExpressionEvaluator.h"
 #include "Math/UnitConversion.h"
+#include "Math/UnrealMathUtility.h"
 #include "Misc/App.h"
 #include "Misc/ExpressionParser.h"
 #include "Misc/Paths.h"
@@ -897,6 +898,9 @@ void SBuildSelection::BuildGroupSelectionChanged(FBuildSelectionBuildGroupPtr It
 
 	int32 Row = 0;
 	int32 Column = 0;
+	const int32 TotalListCount = PlatformList.Num();
+	// Heuristic for choosing grid rows when laying out rows and columns for the quantity of platforms
+	const int32 MaxRows = FMath::Max(4, FMath::FloorToInt32(FMath::Sqrt(float(TotalListCount))) + 2);
 
 	for (const TSharedPtr<FString>& Platform : PlatformList)
 	{
@@ -1002,7 +1006,7 @@ void SBuildSelection::BuildGroupSelectionChanged(FBuildSelectionBuildGroupPtr It
 				]
 			];
 
-		if (++Row > 3)
+		if (++Row >= MaxRows)
 		{
 			Column++;
 			Row = 0;
@@ -1901,7 +1905,7 @@ TSharedRef<SWidget> SBuildSelection::GetGridPanel()
 		})
 		+SHorizontalBox::Slot()
 		.HAlign(HAlign_Left)
-		.FillWidth(0.5f)
+		.AutoWidth()
 		[
 			SNew(SVerticalBox)
 			+SVerticalBox::Slot()
