Allow Build Storage Tool to respect the specified branch and projectfile commandline arguments that can be passed by UGS when launching it.

diff --git a/Engine/Source/Editor/StorageServerWidgets/Private/SBuildSelection.cpp b/Engine/Source/Editor/StorageServerWidgets/Private/SBuildSelection.cpp
--- a/Engine/Source/Editor/StorageServerWidgets/Private/SBuildSelection.cpp
+++ b/Engine/Source/Editor/StorageServerWidgets/Private/SBuildSelection.cpp
@@ -66,6 +66,31 @@ void SBuildSelection::Construct(const FArguments& InArgs)
 		MakeKnownBuildTypePattern(TEXT(".*installed-?build-?(.*)"), EBuildType::EditorPreCompiledBinary, LOCTEXT("BuildSelection_BuildType_PreCompiledBinary", "Editor Pre-compiled Binary"))
 	};
 
+	if (IDesktopPlatform* DesktopPlatform = FDesktopPlatformModule::Get())
+	{
+		DesktopPlatform->EnumerateEngineInstallations(EngineInstallations);
+	}
+
+	// If the tool was launched with an explicit project file, use it to set the engine installation directory and the project
+	// regardless of past settings.  These values should then persist as new settings.
+	FString ProjectFile;
+	if (FParse::Value(FCommandLine::Get(), TEXT("-ProjectFile"), ProjectFile))
+	{
+		if (IDesktopPlatform* DesktopPlatform = FDesktopPlatformModule::Get())
+		{
+			for (const TPair<FString, FString>& EngineInstallation : EngineInstallations)
+			{
+				if (FPaths::IsUnderDirectory(ProjectFile, EngineInstallation.Value))
+				{
+					WriteSetting(TEXT("UserSelectedEngineInstallation"), EngineInstallation.Key);
+					SetUserSelectedProjectDictionaryRoot(EngineInstallation.Value);
+					break;
+				}
+			}
+		}
+		SetSelectedProject(MakeShared<FString>(FPaths::GetBaseFilename(ProjectFile)));
+	}
+
 	FString UserSelectedEngineInstallation;
 	if (ReadSetting(TEXT("UserSelectedEngineInstallation"), UserSelectedEngineInstallation))
 	{
@@ -84,14 +109,6 @@ void SBuildSelection::Construct(const FArguments& InArgs)
 		ReadSetting(LexToString((EBuildType)BuildTypeIndex), TEXT("UserSelectedDestination"), UserSelectedDestinations[BuildTypeIndex]);
 	}
 
-	if (EngineInstallations.IsEmpty())
-	{
-		if (IDesktopPlatform* DesktopPlatform = FDesktopPlatformModule::Get())
-		{
-			DesktopPlatform->EnumerateEngineInstallations(EngineInstallations);
-		}
-	}
-
 	for (const TPair<FString, FString>& EngineInstallation : EngineInstallations)
 	{
 		FallbackProjectDictionary = MakeUnique<FUProjectDictionary>(EngineInstallation.Value);
@@ -116,6 +133,13 @@ void SBuildSelection::Construct(const FArguments& InArgs)
 		EffectiveCompatibleChangelist = CurrentBuildVersion.GetEffectiveCompatibleChangelist();
 	}
 
+	// If the tool was launched with an explicit branch, use it to set the branch name
+	// regardless of past settings.
+	if (FParse::Value(FCommandLine::Get(), TEXT("-Branch"), BranchName))
+	{
+		BranchName = SanitizeBucketSegment(BranchName);
+	}
+
 	if (BranchName.IsEmpty())
 	{
 		ReadSetting(TEXT("Branch"), BranchName);
@@ -1078,14 +1102,6 @@ void SBuildSelection::SetUserSelectedDestination(const FStringView InDestination
 
 	if (BuildType != EBuildType::Oplog)
 	{
-		if (EngineInstallations.IsEmpty())
-		{
-			if (IDesktopPlatform* DesktopPlatform = FDesktopPlatformModule::Get())
-			{
-				DesktopPlatform->EnumerateEngineInstallations(EngineInstallations);
-			}
-		}
-
 		for (const TPair<FString, FString>& EngineInstallation : EngineInstallations)
 		{
 			if (FPaths::IsUnderDirectory(UserSelectedDestinations[BuildTypeIndex], EngineInstallation.Value))
