BuildStorageTool: Change the presentation of builds to show the individual builds using generic artifact naming rather than assuming a one-to-one connection with platform.  Allows removal of temporary configuration splitting and bulk build filtering.

diff --git a/Engine/Source/Developer/Zen/Private/BuildListRetriever.cpp b/Engine/Source/Developer/Zen/Private/BuildListRetriever.cpp
--- a/Engine/Source/Developer/Zen/Private/BuildListRetriever.cpp
+++ b/Engine/Source/Developer/Zen/Private/BuildListRetriever.cpp
@@ -153,7 +153,7 @@ void FBuildListRetriever::IssueBuildQueries()
 
 		for (const FNamespacePlatformBucketTuple& NamespacePlatformBucket : NamespacePlatformBucketTuples)
 		{
-			ServiceInstance->ListBuilds(NamespacePlatformBucket.Namespace, NamespacePlatformBucket.Bucket,
+			ServiceInstance->ListBuilds(NamespacePlatformBucket.Namespace, NamespacePlatformBucket.Bucket, FCbObject(),
 				[this, QueryIndex, Namespace = NamespacePlatformBucket.Namespace, Platform = NamespacePlatformBucket.Platform,
 				PendingQueryState]
 				(TArray<FBuildServiceInstance::FBuildRecord>&& Results) mutable
diff --git a/Engine/Source/Developer/Zen/Private/BuildServerInterface.cpp b/Engine/Source/Developer/Zen/Private/BuildServerInterface.cpp
--- a/Engine/Source/Developer/Zen/Private/BuildServerInterface.cpp
+++ b/Engine/Source/Developer/Zen/Private/BuildServerInterface.cpp
@@ -147,7 +147,7 @@ void FAccessToken::SetToken(const FStringView Scheme, const FStringView Token)
 }
 } // namespace Private
 
-bool LoadFromCompactBinary(FCbFieldView Field, FBuildServiceInstance::FBuildRecord& OutBuildRecord)
+bool LoadFromCompactBinary(FCbFieldView Field, FStringView BucketId, FBuildServiceInstance::FBuildRecord& OutBuildRecord)
 {
 	if (!Field.IsObject())
 	{
@@ -157,6 +157,14 @@ bool LoadFromCompactBinary(FCbFieldView Field, FBuildServiceInstance::FBuildReco
 	FString BuildIdString;
 	FCbObjectId::ByteArray BuildIdBytes;
 	bOk &= LoadFromCompactBinary(Field["buildId"], BuildIdString);
+	if (BucketId.IsEmpty())
+	{
+		bOk &= LoadFromCompactBinary(Field["bucketId"], OutBuildRecord.BucketId);
+	}
+	else
+	{
+		OutBuildRecord.BucketId = BucketId;
+	}
 	bOk &= UE::String::HexToBytes(BuildIdString, BuildIdBytes) == sizeof(BuildIdBytes);
 	OutBuildRecord.BuildId = FCbObjectId(BuildIdBytes);
 
@@ -332,7 +340,7 @@ FBuildServiceInstance::RefreshNamespacesAndBuckets(FOnRefreshNamespacesAndBucket
 }
 
 void
-FBuildServiceInstance::ListBuilds(FStringView Namespace, FStringView Bucket, FOnListBuildsComplete&& InOnListBuildsComplete)
+FBuildServiceInstance::ListBuilds(FStringView Namespace, FStringView Bucket, FCbObject&& Query, FOnListBuildsComplete&& InOnListBuildsComplete)
 {
 	if (ConnectionState.load(std::memory_order_relaxed) != EConnectionState::ConnectionSucceeded)
 	{
@@ -349,22 +357,34 @@ FBuildServiceInstance::ListBuilds(FStringView Namespace, FStringView Bucket, FOn
 	FString OutputFilePath = FPaths::CreateTempFilename(FPlatformProcess::UserTempDir(), TEXT("BuildList"), TEXT(".cbo"));
 	FPaths::MakePlatformFilename(OutputFilePath);
 
-	FString CommandLineArgs = FString::Printf(TEXT("builds list --namespace \"%s\" --bucket \"%s\" \"%s\" \"%s\""),
-		*WriteToString<64>(Namespace), *WriteToString<64>(Bucket), *QueryFilePath, *OutputFilePath);
+	FString BucketParam;
+	if (!Bucket.IsEmpty())
+	{
+		BucketParam = FString::Printf(TEXT(" --bucket \"%s\""), *WriteToString<64>(Bucket));
+	}
+
+	FString CommandLineArgs = FString::Printf(TEXT("builds list --namespace \"%s\"%s \"%s\" \"%s\""),
+		*WriteToString<64>(Namespace), *BucketParam, *QueryFilePath, *OutputFilePath);
 	CommandLineArgs = AddZenUtilityBuildServerArguments(*CommandLineArgs);
 
-	auto PreInvoke = [QueryFilePath](FString& CommandLineArgs)
+	auto PreInvoke = [QueryFilePath, Query = MoveTemp(Query)](FString& CommandLineArgs)
 	{
 		TUniquePtr<FArchive> QueryFileArchive(IFileManager::Get().CreateFileWriter(*QueryFilePath, FILEWRITE_NoFail));
 		FCbWriter Writer;
-		Writer.BeginObject("query");
-		Writer.EndObject();
+		Writer.AddObject("query", Query);
 		Writer.Save(*QueryFileArchive);
 	};
 
 	InvokeZenUtility(CommandLineArgs, MoveTemp(PreInvoke),
-	[OutputFilePath = MoveTemp(OutputFilePath), InOnListBuildsComplete = MoveTemp(InOnListBuildsComplete)](bool bSuccessful) mutable
+	[QueryFilePath = MoveTemp(QueryFilePath), OutputFilePath = MoveTemp(OutputFilePath), Bucket = FString(Bucket), InOnListBuildsComplete = MoveTemp(InOnListBuildsComplete)](bool bSuccessful) mutable
 	{
+		ON_SCOPE_EXIT
+		{
+			IFileManager & FileManager = IFileManager::Get();
+			FileManager.Delete(*QueryFilePath, false, false, true);
+			FileManager.Delete(*OutputFilePath, false, false, true);
+		};
+
 		if (!bSuccessful)
 		{
 			if (InOnListBuildsComplete)
@@ -392,14 +412,13 @@ FBuildServiceInstance::ListBuilds(FStringView Namespace, FStringView Bucket, FOn
 			for (FCbFieldView ResultField : OutputObject["results"].AsArrayView())
 			{
 				FBuildRecord BuildRecord;
-				if (LoadFromCompactBinary(ResultField, BuildRecord))
+				if (LoadFromCompactBinary(ResultField, Bucket, BuildRecord))
 				{
 					BuildRecords.Emplace(MoveTemp(BuildRecord));
 				}
 			}
 			InOnListBuildsComplete(MoveTemp(BuildRecords));
 		}
-
 	});
 }
 
diff --git a/Engine/Source/Developer/Zen/Public/Experimental/BuildServerInterface.h b/Engine/Source/Developer/Zen/Public/Experimental/BuildServerInterface.h
--- a/Engine/Source/Developer/Zen/Public/Experimental/BuildServerInterface.h
+++ b/Engine/Source/Developer/Zen/Public/Experimental/BuildServerInterface.h
@@ -80,6 +80,7 @@ public:
 	struct FBuildRecord
 	{
 		FCbObjectId BuildId;
+		FString BucketId;
 		FCbObject Metadata;
 
 		UE_API FString GetCommitIdentifier() const;
@@ -131,7 +132,7 @@ public:
 
 	inline const TMultiMap<FString, FString>& GetNamespacesAndBuckets() const { return NamespacesAndBuckets; }
 
-	UE_API void ListBuilds(FStringView Namespace, FStringView Bucket, FOnListBuildsComplete&& InOnListBuildsComplete);
+	UE_API void ListBuilds(FStringView Namespace, FStringView Bucket, FCbObject&& Query, FOnListBuildsComplete&& InOnListBuildsComplete);
 
 	UE_API FBuildTransfer StartBuildTransfer(const FCbObjectId& BuildId, FStringView Name, FStringView DestinationFolder, FStringView Namespace, FStringView Bucket);
 	UE_API FBuildTransfer StartOplogBuildTransfer(const FCbObjectId& BuildId, FStringView Name, FStringView DestinationProjectId, FStringView DestinationOplogId, FStringView ProjectFilePath, FStringView Namespace, FStringView Bucket, FName TargetPlatformName);
diff --git a/Engine/Source/Editor/StorageServerWidgets/Internal/SBuildSelection.h b/Engine/Source/Editor/StorageServerWidgets/Internal/SBuildSelection.h
--- a/Engine/Source/Editor/StorageServerWidgets/Internal/SBuildSelection.h
+++ b/Engine/Source/Editor/StorageServerWidgets/Internal/SBuildSelection.h
@@ -54,7 +54,7 @@ class SBuildSelection : public SCompoundWidget
 		FDateTime CreatedAt;
 		FString Category;
 		FString Job;
-		TMap<FString, UE::Zen::Build::FBuildServiceInstance::FBuildRecord> PerPlatformBuilds;
+		TMap<FString, UE::Zen::Build::FBuildServiceInstance::FBuildRecord> NamedArtifacts;
 		bool bIsPreflight;
 	};
 
@@ -145,8 +145,8 @@ private:
 	TSharedPtr<SComboBox<TSharedPtr<FString>>> BuildTypeWidget;
 	TSharedPtr<SMultiSelectComboBox> PlatformFiltersWidget;
 
-	TSharedPtr<SGridPanel> SelectedGroupPlatformGrid;
-	TArray<FString> SelectedGroupSelectedPlatforms;
+	TSharedPtr<SGridPanel> SelectedGroupArtifactGrid;
+	TArray<FString> SelectedGroupSelectedArtifacts;
 
 	TMultiMap<FString, FString> BucketsToNamespaces;
 
@@ -155,6 +155,8 @@ private:
 	TArray<TSharedPtr<FString>> BuildTypeList;
 	TArray<TSharedPtr<FString>> PlatformList;
 
+	TSet<FString> ArtifactNameSet;
+
 	TArray<FString> ActivePlatformFilters;
 	FString UserSelectedDestinations[(int)EBuildType::Count];
 	bool bAppendBuildNameToDestinations[(int)EBuildType::Count];
diff --git a/Engine/Source/Editor/StorageServerWidgets/Private/SBuildSelection.cpp b/Engine/Source/Editor/StorageServerWidgets/Private/SBuildSelection.cpp
--- a/Engine/Source/Editor/StorageServerWidgets/Private/SBuildSelection.cpp
+++ b/Engine/Source/Editor/StorageServerWidgets/Private/SBuildSelection.cpp
@@ -88,7 +88,7 @@ void SBuildSelection::Construct(const FArguments& InArgs)
 				}
 			}
 		}
-		SetSelectedProject(MakeShared<FString>(FPaths::GetBaseFilename(ProjectFile)));
+		SetSelectedProject(MakeShared<FString>(FPaths::GetBaseFilename(ProjectFile).ToLower()));
 	}
 
 	FString UserSelectedEngineInstallation;
@@ -427,7 +427,7 @@ void SBuildSelection::RebuildLists()
 
 		for (const FNamespacePlatformBucketTuple& NamespacePlatformBucket : NamespacePlatformBucketTuples)
 		{
-			ServiceInstance->ListBuilds(NamespacePlatformBucket.Namespace, NamespacePlatformBucket.Bucket,
+			ServiceInstance->ListBuilds(NamespacePlatformBucket.Namespace, NamespacePlatformBucket.Bucket, FCbObject(),
 			[this, QueryIndex, Namespace = NamespacePlatformBucket.Namespace, Platform = NamespacePlatformBucket.Platform,
 				ExpectedBuildRefreshGeneration = BuildRefreshGeneration.load(), PendingQueryState]
 			(TArray<FBuildServiceInstance::FBuildRecord>&& Results) mutable
@@ -483,29 +483,13 @@ void SBuildSelection::RegenerateBuildGroups(UE::Zen::Build::FListBuildsState& Li
 	using namespace UE::BuildSelection::Internal;
 
 	BuildGroups.Empty();
+	ArtifactNameSet.Empty();
 
-	const bool bDoConfigurationSplitting = GetSelectedBuildType() == EBuildType::PackagedBuild;
-
-	auto MakeGroupKey = [bDoConfigurationSplitting](const FString& Namespace, const FString& CommitIdentifier, const FBuildServiceInstance::FBuildRecord& BuildRecord)
+	auto MakeGroupKey = [](const FString& Namespace, const FString& CommitIdentifier, const FBuildServiceInstance::FBuildRecord& BuildRecord)
 	{
-		// TODO: Append the configuration to the group key if present to workaround  issues with multiple uploads with the same build group but separate configuration
-		// Should be removed once we improve the platform handling to be more flexible and thus include configurations
-		FString Configuration;
-		if (FCbFieldView ConfigurationField = BuildRecord.Metadata["Configuration"]; ConfigurationField.HasValue() && !ConfigurationField.HasError())
-		{
-			Configuration = FUTF8ToTCHAR(ConfigurationField.AsString());
-		}
-
 		if (FCbFieldView BuildGroupField = BuildRecord.Metadata["buildgroup"]; BuildGroupField.HasValue() && !BuildGroupField.HasError())
 		{
-			if (Configuration.IsEmpty() || !bDoConfigurationSplitting) 
-			{
-				return FString(WriteToString<64>(Namespace, ".", BuildGroupField.AsString()));
-			}
-			else
-			{
-				return FString(WriteToString<64>(Namespace, ".", BuildGroupField.AsString(), ".", Configuration));
-			}
+			return FString(WriteToString<64>(Namespace, ".", BuildGroupField.AsString()));
 		}
 		if (FCbFieldView NameField = BuildRecord.Metadata["name"]; NameField.HasValue() && !NameField.HasError())
 		{
@@ -532,22 +516,22 @@ void SBuildSelection::RegenerateBuildGroups(UE::Zen::Build::FListBuildsState& Li
 				BuildGroup = MakeShared<FBuildGroup>();
 			}
 
-			// TODO: Temporary detection of bulk builds to work around issue were a build can have multiple versions of the same platform
-			bool IsBulkBuild = false;
+			FString ArtifactName;
 			if (FCbFieldView NameView = BuildRecord.Metadata["name"]; NameView.HasValue() && !NameView.HasError())
 			{
-				FString Name;
-				Name = FUTF8ToTCHAR(NameView.AsString());
-				int32 BulkStartIndex = Name.Find(TEXT("-Bulk"));
-				if (BulkStartIndex != INDEX_NONE)
+				ArtifactName = FUTF8ToTCHAR(NameView.AsString());
+
+				//Temporary accommodation for oplog names that need to be changed to adhere to the new naming expectations
+				int32 DotIndex;
+				if (ArtifactName.FindLastChar(TCHAR('.'), DotIndex))
 				{
-					IsBulkBuild = true;
+					ArtifactName.RightChopInline(DotIndex + 1);
 				}
 			}
 
-			if (IsBulkBuild)
+			if (ArtifactName.IsEmpty())
 			{
-				 continue;
+				ArtifactName = BuildState.Platform;
 			}
 			
 			if (BuildGroup->DisplayName.IsEmpty())
@@ -627,12 +611,6 @@ void SBuildSelection::RegenerateBuildGroups(UE::Zen::Build::FListBuildsState& Li
 						ItemName.RightChopInline(Category.Len());
 					}
 
-					// TODO: Configuration appending is temporary until we have improved the platform grouping to include configuration
-					if (bDoConfigurationSplitting && !Configuration.IsEmpty() && !ItemName.EndsWith(Configuration))
-					{
-						// append the configuration to the item name if present
-						ItemName = *WriteToString<64>(ItemName, "-", Configuration);
-					}
 					bool bCharRemoved = false;
 					do
 					{
@@ -662,7 +640,8 @@ void SBuildSelection::RegenerateBuildGroups(UE::Zen::Build::FListBuildsState& Li
 				BuildGroup->bIsPreflight = bIsPreflight;
 			}
 
-			BuildGroup->PerPlatformBuilds.FindOrAdd(BuildState.Platform, MoveTemp(BuildRecord));
+			ArtifactNameSet.Add(ArtifactName);
+			BuildGroup->NamedArtifacts.FindOrAdd(ArtifactName, MoveTemp(BuildRecord));
 		}
 	}
 
@@ -758,11 +737,6 @@ void SBuildSelection::RegenerateActivePlatformFilters()
 			if (PlatformFiltersWidget->IsChecked(*Platform))
 			{
 				ActivePlatformFilters.Add(*Platform);
-				SelectedGroupSelectedPlatforms.AddUnique(*Platform);
-			}
-			else
-			{
-				SelectedGroupSelectedPlatforms.Remove(*Platform);
 			}
 		}
 	}
@@ -806,7 +780,18 @@ bool SBuildSelection::DoesPassAllViewFilters(TSharedPtr<FBuildGroup> InItem) con
 
 	for (const FString& ActivePlatformFilter : ActivePlatformFilters)
 	{
-		if (!InItem->PerPlatformBuilds.Contains(ActivePlatformFilter))
+		bool bHasPlatform = false;
+		for (const TPair<FString, UE::Zen::Build::FBuildServiceInstance::FBuildRecord>& NamedArtifact : InItem->NamedArtifacts)
+		{
+			if (NamedArtifact.Value.BucketId.EndsWith(ActivePlatformFilter) &&
+				NamedArtifact.Value.BucketId.Len() > ActivePlatformFilter.Len() &&
+				NamedArtifact.Value.BucketId[NamedArtifact.Value.BucketId.Len() - ActivePlatformFilter.Len() - 1] == TCHAR('.'))
+			{
+				bHasPlatform = true;
+				break;
+			}
+		}
+		if (!bHasPlatform)
 		{
 			return false;
 		}
@@ -908,12 +893,12 @@ void SBuildSelection::BuildGroupSelectionChanged(FBuildSelectionBuildGroupPtr It
 {
 	using namespace UE::Zen::Build;
 
-	if (!SelectedGroupPlatformGrid)
+	if (!SelectedGroupArtifactGrid)
 	{
 		return;
 	}
 
-	SelectedGroupPlatformGrid->ClearChildren();
+	SelectedGroupArtifactGrid->ClearChildren();
 
 	if (!Item)
 	{
@@ -922,40 +907,40 @@ void SBuildSelection::BuildGroupSelectionChanged(FBuildSelectionBuildGroupPtr It
 
 	int32 Row = 0;
 	int32 Column = 0;
-	const int32 TotalListCount = PlatformList.Num();
+	const int32 TotalListCount = ArtifactNameSet.Num();
 	// Heuristic for choosing grid rows when laying out rows and columns for the quantity of platforms
 	const int32 MaxRows = FMath::Max(4, FMath::FloorToInt32(FMath::Sqrt(float(TotalListCount))) + 2);
 
-	for (const TSharedPtr<FString>& Platform : PlatformList)
+	for (const FString& ArtifactName : ArtifactNameSet)
 	{
 		TSharedPtr<SCheckBox> CurrentCheckbox;
-		SelectedGroupPlatformGrid->AddSlot(Column, Row)
+		SelectedGroupArtifactGrid->AddSlot(Column, Row)
 			.Padding(0,0,0,2)
 			[
 				SNew(SHorizontalBox)
-				.IsEnabled_Lambda([this, Item, Platform]
+				.IsEnabled_Lambda([this, Item, ArtifactName]
 				{
-					return Item->PerPlatformBuilds.Contains(*Platform);
+					return Item->NamedArtifacts.Contains(*ArtifactName);
 				})
 				+SHorizontalBox::Slot()
 				.AutoWidth()
 				.VAlign(VAlign_Center)
 				[
 					SAssignNew(CurrentCheckbox, SCheckBox)
-					.IsChecked_Lambda([this, Item, Platform]
+					.IsChecked_Lambda([this, Item, ArtifactName]
 					{
-						bool bIsChecked = Item->PerPlatformBuilds.Contains(*Platform) && SelectedGroupSelectedPlatforms.Contains(*Platform);
+						bool bIsChecked = Item->NamedArtifacts.Contains(ArtifactName) && SelectedGroupSelectedArtifacts.Contains(ArtifactName);
 						return bIsChecked ? ECheckBoxState::Checked : ECheckBoxState::Unchecked;
 					})
-					.OnCheckStateChanged_Lambda([this, Platform](ECheckBoxState InNewState)
+					.OnCheckStateChanged_Lambda([this, ArtifactName](ECheckBoxState InNewState)
 					{
 						if (InNewState == ECheckBoxState::Checked)
 						{
-							SelectedGroupSelectedPlatforms.AddUnique(*Platform);
+							SelectedGroupSelectedArtifacts.AddUnique(ArtifactName);
 						}
 						else if (InNewState == ECheckBoxState::Unchecked)
 						{
-							SelectedGroupSelectedPlatforms.Remove(*Platform);
+							SelectedGroupSelectedArtifacts.Remove(ArtifactName);
 						}
 					})
 				]
@@ -967,7 +952,7 @@ void SBuildSelection::BuildGroupSelectionChanged(FBuildSelectionBuildGroupPtr It
 					.ButtonStyle(FAppStyle::Get(), "InvisibleButton")
 					.IsFocusable(false)
 
-					.OnClicked_Lambda([this, CurrentCheckbox, Platform]()
+					.OnClicked_Lambda([this, CurrentCheckbox, ArtifactName]()
 						{
 							ECheckBoxState NewState = ECheckBoxState::Checked;
 							if (CurrentCheckbox->IsChecked())
@@ -977,18 +962,18 @@ void SBuildSelection::BuildGroupSelectionChanged(FBuildSelectionBuildGroupPtr It
 							CurrentCheckbox->SetIsChecked(NewState);
 							if (NewState == ECheckBoxState::Checked)
 							{
-								SelectedGroupSelectedPlatforms.AddUnique(*Platform);
+								SelectedGroupSelectedArtifacts.AddUnique(ArtifactName);
 							}
 							else if (NewState == ECheckBoxState::Unchecked)
 							{
-								SelectedGroupSelectedPlatforms.Remove(*Platform);
+								SelectedGroupSelectedArtifacts.Remove(ArtifactName);
 							}
 							return FReply::Handled();
 						})
 					[
 						SNew(STextBlock)
 						.Justification(ETextJustify::Left)
-						.Text(FText::FromString(*Platform))
+						.Text(FText::FromString(ArtifactName))
 					]
 				]
 				+SHorizontalBox::Slot()
@@ -998,9 +983,9 @@ void SBuildSelection::BuildGroupSelectionChanged(FBuildSelectionBuildGroupPtr It
 					SNew(SButton)
 					.ButtonStyle(&FAppStyle::Get().GetWidgetStyle<FButtonStyle>("SimpleButton"))
 					.ToolTipText(LOCTEXT("BuildSelection_BrowseToBuildUrl", "Browse to build url"))
-					.IsEnabled_Lambda([this, Item, Platform]
+					.IsEnabled_Lambda([this, Item, ArtifactName]
 					{
-						if (FBuildServiceInstance::FBuildRecord* BuildRecord = Item->PerPlatformBuilds.Find(*Platform))
+						if (FBuildServiceInstance::FBuildRecord* BuildRecord = Item->NamedArtifacts.Find(ArtifactName))
 						{
 							if (FCbFieldView BuildUrlField = BuildRecord->Metadata["buildurl"]; BuildUrlField.HasValue() && !BuildUrlField.HasError())
 							{
@@ -1009,9 +994,9 @@ void SBuildSelection::BuildGroupSelectionChanged(FBuildSelectionBuildGroupPtr It
 						}
 						return false;
 					})
-					.OnClicked_Lambda([this, Item, Platform]
+					.OnClicked_Lambda([this, Item, ArtifactName]
 					{
-						if (FBuildServiceInstance::FBuildRecord* BuildRecord = Item->PerPlatformBuilds.Find(*Platform))
+						if (FBuildServiceInstance::FBuildRecord* BuildRecord = Item->NamedArtifacts.Find(ArtifactName))
 						{
 							if (FCbFieldView BuildUrlField = BuildRecord->Metadata["buildurl"]; BuildUrlField.HasValue() && !BuildUrlField.HasError())
 							{
@@ -1932,14 +1917,14 @@ TSharedRef<SWidget> SBuildSelection::GetGridPanel()
 				SNew(STextBlock)
 					.ColorAndOpacity(TitleColor)
 					.Font(TitleFont)
-					.Text(LOCTEXT("BuildSelection_AvailablePlatforms", "Available Platforms"))
+					.Text(LOCTEXT("BuildSelection_AvailableBuilds", "Available Builds"))
 			]
 			+SVerticalBox::Slot()
 			.HAlign(HAlign_Fill)
 			.VAlign(VAlign_Top)
 			.AutoHeight()
 			[
-				SAssignNew(SelectedGroupPlatformGrid, SGridPanel)
+				SAssignNew(SelectedGroupArtifactGrid, SGridPanel)
 			]
 		]
 		+SHorizontalBox::Slot()
@@ -2010,9 +1995,9 @@ TSharedRef<SWidget> SBuildSelection::GetGridPanel()
 						return false;
 					}
 
-					for (const FString& PlatformForSelectedGroup : SelectedGroupSelectedPlatforms)
+					for (const FString& ArtifactForSelectedGroup : SelectedGroupSelectedArtifacts)
 					{
-						if (SelectedItems[0]->PerPlatformBuilds.Contains(PlatformForSelectedGroup))
+						if (SelectedItems[0]->NamedArtifacts.Contains(ArtifactForSelectedGroup))
 						{
 							return true;
 						}
@@ -2031,13 +2016,11 @@ TSharedRef<SWidget> SBuildSelection::GetGridPanel()
 							return FReply::Handled();
 						}
 
-						for (const FString& PlatformForSelectedGroup : SelectedGroupSelectedPlatforms)
+						for (const FString& ArtifactForSelectedGroup : SelectedGroupSelectedArtifacts)
 						{
-							if (FBuildServiceInstance::FBuildRecord* BuildRecord = SelectedItems[0]->PerPlatformBuilds.Find(PlatformForSelectedGroup))
+							if (FBuildServiceInstance::FBuildRecord* BuildRecord = SelectedItems[0]->NamedArtifacts.Find(ArtifactForSelectedGroup))
 							{
-								FString Bucket = FString::Printf(TEXT("%s.%s.%s.%s"), **SelectedProject, **SelectedBuildType, **SelectedStream, *PlatformForSelectedGroup);
-
-								FString DestinationPlatformName = PlatformForSelectedGroup;
+								FString DestinationPlatformName = ArtifactForSelectedGroup;
 								if (FCbFieldView CookPlatformField = BuildRecord->Metadata["cookPlatform"]; CookPlatformField.HasValue() && !CookPlatformField.HasError())
 								{
 									DestinationPlatformName = *WriteToString<64>(CookPlatformField.AsString());
@@ -2065,8 +2048,8 @@ TSharedRef<SWidget> SBuildSelection::GetGridPanel()
 										DestinationOplogId = FString::Printf(TEXT("%s.%s"), *DestinationOplogId, *SanitizeForZenId(SelectedItems[0]->DisplayName));
 									}
 									FBuildServiceInstance::FBuildTransfer BuildTransfer =
-										ServiceInstance->StartOplogBuildTransfer(BuildRecord->BuildId, TransferName, DestinationProjectId, DestinationOplogId, ProjectFilePath, SelectedItems[0]->Namespace, Bucket, TargetPlatformName);
-									OnBuildTransferStarted.ExecuteIfBound(BuildTransfer, SelectedItems[0]->DisplayName, PlatformForSelectedGroup);
+										ServiceInstance->StartOplogBuildTransfer(BuildRecord->BuildId, TransferName, DestinationProjectId, DestinationOplogId, ProjectFilePath, SelectedItems[0]->Namespace, BuildRecord->BucketId, TargetPlatformName);
+									OnBuildTransferStarted.ExecuteIfBound(BuildTransfer, SelectedItems[0]->DisplayName, ArtifactForSelectedGroup);
 								}
 								else
 								{
@@ -2080,8 +2063,8 @@ TSharedRef<SWidget> SBuildSelection::GetGridPanel()
 										DestinationFolder = FPaths::Combine(GetEffectiveDestination(), DestinationPlatformName);
 									}
 									FBuildServiceInstance::FBuildTransfer BuildTransfer =
-										ServiceInstance->StartBuildTransfer(BuildRecord->BuildId, TransferName, DestinationFolder, SelectedItems[0]->Namespace, Bucket);
-									OnBuildTransferStarted.ExecuteIfBound(BuildTransfer, SelectedItems[0]->DisplayName, PlatformForSelectedGroup);
+										ServiceInstance->StartBuildTransfer(BuildRecord->BuildId, TransferName, DestinationFolder, SelectedItems[0]->Namespace, BuildRecord->BucketId);
+									OnBuildTransferStarted.ExecuteIfBound(BuildTransfer, SelectedItems[0]->DisplayName, ArtifactForSelectedGroup);
 								}
 							}
 						}
