BuildStorageTool: Fixes to grouping and artifact naming to accommodate legacy metadata.

diff --git a/Engine/Source/Editor/StorageServerWidgets/Private/SBuildSelection.cpp b/Engine/Source/Editor/StorageServerWidgets/Private/SBuildSelection.cpp
--- a/Engine/Source/Editor/StorageServerWidgets/Private/SBuildSelection.cpp
+++ b/Engine/Source/Editor/StorageServerWidgets/Private/SBuildSelection.cpp
@@ -493,7 +493,18 @@ void SBuildSelection::RegenerateBuildGroups(UE::Zen::Build::FListBuildsState& Li
 		}
 		if (FCbFieldView NameField = BuildRecord.Metadata["name"]; NameField.HasValue() && !NameField.HasError())
 		{
-			return FString(WriteToString<64>(Namespace, ".", NameField.AsString()));
+			FString CandidateKeyName = FString(WriteToString<64>(Namespace, ".", NameField.AsString()));
+			// TODO: This name manipulation needs to be removed when the metadata is more consistent.
+			int32 CLIndex = CandidateKeyName.Find(TEXT("-CL-"), ESearchCase::IgnoreCase, ESearchDir::FromEnd);
+			if (CLIndex != INDEX_NONE)
+			{
+				CLIndex = CandidateKeyName.Find(TEXT("-"), ESearchCase::IgnoreCase, ESearchDir::FromStart, CLIndex + 4);
+				if (CLIndex != INDEX_NONE)
+				{
+					CandidateKeyName.LeftInline(CLIndex);
+				}
+			}
+			return CandidateKeyName;
 		}
 		if (!CommitIdentifier.IsEmpty())
 		{
@@ -521,11 +532,20 @@ void SBuildSelection::RegenerateBuildGroups(UE::Zen::Build::FListBuildsState& Li
 			{
 				ArtifactName = FUTF8ToTCHAR(NameView.AsString());
 
-				//Temporary accommodation for oplog names that need to be changed to adhere to the new naming expectations
-				int32 DotIndex;
-				if (ArtifactName.FindLastChar(TCHAR('.'), DotIndex))
+				// TODO: This name manipulation needs to be removed when the metadata is more consistent.
+				int32 TruncationIndex;
+				int32 CLIndex = ArtifactName.Find(TEXT("-CL-"), ESearchCase::IgnoreCase, ESearchDir::FromEnd);
+				if (CLIndex != INDEX_NONE)
+				{
+					ArtifactName.RightChopInline(CLIndex + 4);
+					if (ArtifactName.FindChar(TCHAR('-'), TruncationIndex))
+					{
+						ArtifactName.RightChopInline(TruncationIndex + 1);
+					}
+				}
+				if (ArtifactName.FindLastChar(TCHAR('.'), TruncationIndex))
 				{
-					ArtifactName.RightChopInline(DotIndex + 1);
+					ArtifactName.RightChopInline(TruncationIndex + 1);
 				}
 			}
 
