BuildStorageTool: More fixes to grouping and artifact naming to accommodate legacy metadata.

diff --git a/Engine/Source/Editor/StorageServerWidgets/Private/SBuildSelection.cpp b/Engine/Source/Editor/StorageServerWidgets/Private/SBuildSelection.cpp
--- a/Engine/Source/Editor/StorageServerWidgets/Private/SBuildSelection.cpp
+++ b/Engine/Source/Editor/StorageServerWidgets/Private/SBuildSelection.cpp
@@ -498,10 +498,18 @@ void SBuildSelection::RegenerateBuildGroups(UE::Zen::Build::FListBuildsState& Li
 			int32 CLIndex = CandidateKeyName.Find(TEXT("-CL-"), ESearchCase::IgnoreCase, ESearchDir::FromEnd);
 			if (CLIndex != INDEX_NONE)
 			{
-				CLIndex = CandidateKeyName.Find(TEXT("-"), ESearchCase::IgnoreCase, ESearchDir::FromStart, CLIndex + 4);
-				if (CLIndex != INDEX_NONE)
+				int32 DashIndex = CandidateKeyName.Find(TEXT("-"), ESearchCase::IgnoreCase, ESearchDir::FromStart, CLIndex + 4);
+				if (DashIndex != INDEX_NONE)
 				{
-					CandidateKeyName.LeftInline(CLIndex);
+					CandidateKeyName.LeftInline(DashIndex);
+				}
+				else
+				{
+					int32 DotIndex = CandidateKeyName.Find(TEXT("."), ESearchCase::IgnoreCase, ESearchDir::FromStart, CLIndex + 4);
+					if (DotIndex != INDEX_NONE)
+					{
+						CandidateKeyName.LeftInline(DotIndex);
+					}
 				}
 			}
 			return CandidateKeyName;
@@ -547,6 +555,11 @@ void SBuildSelection::RegenerateBuildGroups(UE::Zen::Build::FListBuildsState& Li
 				{
 					ArtifactName.RightChopInline(TruncationIndex + 1);
 				}
+
+				if (ArtifactName.Len() < 5)
+				{
+					ArtifactName.Empty();
+				}
 			}
 
 			if (ArtifactName.IsEmpty())
