BuildStorageTool: Add wildcard filtering (includusion and exclusion) when downloading files-based builds (not oplogs/snapshots).

diff --git a/Engine/Source/Developer/Zen/Private/BuildServerInterface.cpp b/Engine/Source/Developer/Zen/Private/BuildServerInterface.cpp
--- a/Engine/Source/Developer/Zen/Private/BuildServerInterface.cpp
+++ b/Engine/Source/Developer/Zen/Private/BuildServerInterface.cpp
@@ -62,6 +62,8 @@ struct FBuildTransferState : public TSharedFromThis<FBuildTransferState, ESPMode
 	FString Bucket;
 	FString Destination;
 	FString ProjectFilePath;
+	FString IncludeFilter;
+	FString ExcludeFilter;
 	FName TargetPlatformName = NAME_None;
 
 	std::atomic<bool> bCancelRequested = false;
@@ -423,7 +425,14 @@ FBuildServiceInstance::ListBuilds(FStringView Namespace, FStringView Bucket, FCb
 }
 
 FBuildServiceInstance::FBuildTransfer
-FBuildServiceInstance::StartBuildTransfer(const FCbObjectId& BuildId, FStringView Name, FStringView DestinationFolder, FStringView Namespace, FStringView Bucket)
+FBuildServiceInstance::StartBuildTransfer(
+	const FCbObjectId& BuildId,
+	FStringView Name,
+	FStringView DestinationFolder,
+	FStringView Namespace,
+	FStringView Bucket,
+	FStringView IncludeFilter,
+	FStringView ExcludeFilter)
 {
 	TSharedPtr<Private::FBuildTransferState> NewState = MakeShared<Private::FBuildTransferState>();
 	NewState->Type = FBuildTransfer::EType::Files;
@@ -432,6 +441,8 @@ FBuildServiceInstance::StartBuildTransfer(const FCbObjectId& BuildId, FStringVie
 	NewState->Namespace = Namespace;
 	NewState->Bucket = Bucket;
 	NewState->Destination = DestinationFolder;
+	NewState->IncludeFilter = IncludeFilter;
+	NewState->ExcludeFilter = ExcludeFilter;
 	BuildTransferThreadStates.Enqueue(NewState.ToSharedRef());
 	KickBuildTransferThread();
 
@@ -441,7 +452,15 @@ FBuildServiceInstance::StartBuildTransfer(const FCbObjectId& BuildId, FStringVie
 }
 
 FBuildServiceInstance::FBuildTransfer
-FBuildServiceInstance::StartOplogBuildTransfer(const FCbObjectId& BuildId, FStringView Name, FStringView DestinationProjectId, FStringView DestinationOplogId, FStringView ProjectFilePath, FStringView Namespace, FStringView Bucket, FName TargetPlatformName)
+FBuildServiceInstance::StartOplogBuildTransfer(
+	const FCbObjectId& BuildId,
+	FStringView Name,
+	FStringView DestinationProjectId,
+	FStringView DestinationOplogId,
+	FStringView ProjectFilePath,
+	FStringView Namespace,
+	FStringView Bucket,
+	FName TargetPlatformName)
 {
 	TSharedPtr<Private::FBuildTransferState> NewState = MakeShared<Private::FBuildTransferState>();
 	NewState->Type = FBuildTransfer::EType::Oplog;
@@ -473,6 +492,8 @@ FBuildServiceInstance::RepeatBuildTransfer(FBuildTransfer BuildTransfer)
 		NewState->Namespace = BuildTransfer.State->Namespace;
 		NewState->Bucket = BuildTransfer.State->Bucket;
 		NewState->Destination = BuildTransfer.State->Destination;
+		NewState->IncludeFilter = BuildTransfer.State->IncludeFilter;
+		NewState->ExcludeFilter = BuildTransfer.State->ExcludeFilter;
 		NewState->ProjectFilePath = BuildTransfer.State->ProjectFilePath;
 		NewState->TargetPlatformName = BuildTransfer.State->TargetPlatformName;
 	}
@@ -942,9 +963,12 @@ FBuildServiceInstance::BuildTransferThreadLoop()
 		{
 		case FBuildTransfer::EType::Files:
 			DestinationDirectory = FPaths::ConvertRelativePathToFull(State->Destination);
-			CommandLineArgs = FString::Printf(TEXT("builds download --log-progress --host %s --local-path \"%s\" --namespace %s --bucket %s --build-id %s --oidctoken-exe-path \"%s\""),
+			CommandLineArgs = FString::Printf(TEXT("builds download --log-progress --host %s --local-path \"%s\" --namespace %s --bucket %s --build-id %s%s%s --oidctoken-exe-path \"%s\""),
 				*WriteToString<64>(EffectiveDomain), *DestinationDirectory, *State->Namespace, *State->Bucket,
-				*WriteToString<64>(State->BuildId), *OidcTokenExecutableFilename);
+				*WriteToString<64>(State->BuildId),
+				State->IncludeFilter.IsEmpty() ? TEXT("") : *WriteToString<64>(TEXT(" --wildcard \""), State->IncludeFilter, TEXT("\"")),
+				State->ExcludeFilter.IsEmpty() ? TEXT("") : *WriteToString<64>(TEXT(" --exclude-wildcard \""), State->ExcludeFilter, TEXT("\"")),
+				*OidcTokenExecutableFilename);
 			break;
 		case FBuildTransfer::EType::Oplog:
 			{
@@ -1009,8 +1033,6 @@ FBuildServiceInstance::BuildTransferThreadLoop()
 					*DestinationProjectId, *DestinationOplogId, *WriteToString<64>(EffectiveDomain),
 					*State->Namespace, *State->Bucket,
 					*WriteToString<64>(State->BuildId), *GetAccessToken());
-
-
 			}
 			break;
 		}
diff --git a/Engine/Source/Developer/Zen/Public/Experimental/BuildServerInterface.h b/Engine/Source/Developer/Zen/Public/Experimental/BuildServerInterface.h
--- a/Engine/Source/Developer/Zen/Public/Experimental/BuildServerInterface.h
+++ b/Engine/Source/Developer/Zen/Public/Experimental/BuildServerInterface.h
@@ -134,8 +134,25 @@ public:
 
 	UE_API void ListBuilds(FStringView Namespace, FStringView Bucket, FCbObject&& Query, FOnListBuildsComplete&& InOnListBuildsComplete);
 
-	UE_API FBuildTransfer StartBuildTransfer(const FCbObjectId& BuildId, FStringView Name, FStringView DestinationFolder, FStringView Namespace, FStringView Bucket);
-	UE_API FBuildTransfer StartOplogBuildTransfer(const FCbObjectId& BuildId, FStringView Name, FStringView DestinationProjectId, FStringView DestinationOplogId, FStringView ProjectFilePath, FStringView Namespace, FStringView Bucket, FName TargetPlatformName);
+	UE_API FBuildTransfer StartBuildTransfer(
+		const FCbObjectId& BuildId,
+		FStringView Name,
+		FStringView DestinationFolder,
+		FStringView Namespace,
+		FStringView Bucket,
+		FStringView IncludeFilter = FStringView(),
+		FStringView ExcludeFilter = FStringView()
+		);
+	UE_API FBuildTransfer StartOplogBuildTransfer(
+		const FCbObjectId& BuildId,
+		FStringView Name,
+		FStringView DestinationProjectId,
+		FStringView DestinationOplogId,
+		FStringView ProjectFilePath,
+		FStringView Namespace,
+		FStringView Bucket,
+		FName TargetPlatformName
+		);
 	UE_API FBuildTransfer RepeatBuildTransfer(FBuildTransfer ExistingTransfer);
 private:
 	typedef TUniqueFunction<void(FString& CommandLineArgs)> FPreZenUtilityInvocation;
diff --git a/Engine/Source/Editor/StorageServerWidgets/Internal/SBuildSelection.h b/Engine/Source/Editor/StorageServerWidgets/Internal/SBuildSelection.h
--- a/Engine/Source/Editor/StorageServerWidgets/Internal/SBuildSelection.h
+++ b/Engine/Source/Editor/StorageServerWidgets/Internal/SBuildSelection.h
@@ -90,7 +90,8 @@ private:
 	bool BuildGroupIsSelectableOrNavigable(TSharedPtr<FBuildGroup> InItem) const;
 	TSharedRef<ITableRow> GenerateBuildGroupRow(TSharedPtr<FBuildGroup> InItem, const TSharedRef<STableViewBase>& InOwningTable);
 	TSharedPtr<SWidget> OnGetBuildGroupContextMenuContent() const;
-	void OnOpenDestinationDirectoryClicked();
+	void OnChooseDestinationDirectoryClicked();
+	void OnChooseTransferFilterClicked();
 	TSharedRef<SWidget> GetBuildDestinationPanel();
 	TSharedRef<SWidget> GetGridPanel();
 	void BuildGroupSelectionChanged(TSharedPtr<FBuildGroup> Item, ESelectInfo::Type SelectInfo);
@@ -105,6 +106,8 @@ private:
 	FString GetDefaultDestination() const;
 	FString GetEffectiveDestination() const;
 
+	void SetUserSelectedTransferFilter(const FStringView InIncludeFilter, const FStringView InExcludeFilter);
+
 	void SetSelectedStream(const TSharedPtr<FString> InSelectedStream);
 	void SetSelectedProject(const TSharedPtr<FString> InSelectedProject);
 	void SetSelectedBuildType(const TSharedPtr<FString> InBuildType);
@@ -125,6 +128,8 @@ private:
 
 	FReply ExploreDestination_OnClicked();
 
+	void DownloadSelectedArtifactsInSelectedBuildGroup();
+
 	static const FString EmptyString;
 
 	TMap<FString, FString> EngineInstallations;
@@ -157,8 +162,15 @@ private:
 
 	TSet<FString> ArtifactNameSet;
 
+	struct FTransferFilter
+	{
+		FString IncludeFilter;
+		FString ExcludeFilter;
+	};
+
 	TArray<FString> ActivePlatformFilters;
 	FString UserSelectedDestinations[(int)EBuildType::Count];
+	FTransferFilter UserSelectedTransferFilters[(int)EBuildType::Count];
 	bool bAppendBuildNameToDestinations[(int)EBuildType::Count];
 
 	TSharedPtr<FString> SelectedStream;
diff --git a/Engine/Source/Editor/StorageServerWidgets/Internal/SBuildTransferFilters.h b/Engine/Source/Editor/StorageServerWidgets/Internal/SBuildTransferFilters.h
new file mode 100644
--- /dev/null
+++ b/Engine/Source/Editor/StorageServerWidgets/Internal/SBuildTransferFilters.h
@@ -0,0 +1,43 @@
+// Copyright Epic Games, Inc. All Rights Reserved.
+
+#pragma once
+
+#include "Containers/Array.h"
+#include "Containers/UnrealString.h"
+#include "Internationalization/Text.h"
+#include "Misc/Guid.h"
+#include "Templates/SharedPointer.h"
+#include "Widgets/DeclarativeSyntaxSupport.h"
+#include "Widgets/SBoxPanel.h"
+#include "Widgets/SCompoundWidget.h"
+#include "Widgets/SWidget.h"
+
+#define UE_API STORAGESERVERWIDGETS_API
+
+class SBuildTransferFilters : public SCompoundWidget
+{
+	SLATE_BEGIN_ARGS(SBuildTransferFilters)
+		: _IncludeFilter(nullptr)
+		, _ExcludeFilter(nullptr)
+	{ }
+
+	SLATE_ATTRIBUTE(TSharedPtr<FString>, IncludeFilter);
+	SLATE_ATTRIBUTE(TSharedPtr<FString>, ExcludeFilter);
+
+	SLATE_END_ARGS()
+
+	void Construct(const FArguments& InArgs);
+
+	FString GetIncludeFilter() const { return IncludeFilter; }
+	FString GetExcludeFilter() const { return ExcludeFilter; }
+
+private:
+	TSharedRef<SWidget> GetGridPanel();
+
+	FString IncludeFilter;
+	FString ExcludeFilter;
+
+	SVerticalBox::FSlot* GridSlot = nullptr;
+};
+
+#undef UE_API
diff --git a/Engine/Source/Editor/StorageServerWidgets/Private/SBuildSelection.cpp b/Engine/Source/Editor/StorageServerWidgets/Private/SBuildSelection.cpp
--- a/Engine/Source/Editor/StorageServerWidgets/Private/SBuildSelection.cpp
+++ b/Engine/Source/Editor/StorageServerWidgets/Private/SBuildSelection.cpp
@@ -3,6 +3,8 @@
 #include "SBuildSelection.h"
 
 #include "DesktopPlatformModule.h"
+#include "Dialog/DialogCommands.h"
+#include "Dialog/SCustomDialog.h"
 #include "Framework/MultiBox/MultiBoxBuilder.h"
 #include "HAL/PlatformApplicationMisc.h"
 #include "HAL/PlatformFileManager.h"
@@ -16,6 +18,7 @@
 #include "Misc/Paths.h"
 #include "Misc/UProjectInfo.h"
 #include "Modules/BuildVersion.h"
+#include "SBuildTransferFilters.h"
 #include "SMultiSelectComboBox.h"
 #include "Styling/StyleColors.h"
 #include "Widgets/Images/SImage.h"
@@ -109,6 +112,12 @@ void SBuildSelection::Construct(const FArguments& InArgs)
 		ReadSetting(LexToString((EBuildType)BuildTypeIndex), TEXT("UserSelectedDestination"), UserSelectedDestinations[BuildTypeIndex]);
 	}
 
+	for (int32 BuildTypeIndex = 0; BuildTypeIndex < (int32)EBuildType::Count; ++BuildTypeIndex)
+	{
+		ReadSetting(LexToString((EBuildType)BuildTypeIndex), TEXT("UserSelectedIncludeFilter"), UserSelectedTransferFilters[BuildTypeIndex].IncludeFilter);
+		ReadSetting(LexToString((EBuildType)BuildTypeIndex), TEXT("UserSelectedExcludeFilter"), UserSelectedTransferFilters[BuildTypeIndex].ExcludeFilter);
+	}
+
 	for (const TPair<FString, FString>& EngineInstallation : EngineInstallations)
 	{
 		FallbackProjectDictionary = MakeUnique<FUProjectDictionary>(EngineInstallation.Value);
@@ -1134,6 +1143,17 @@ void SBuildSelection::SetUserSelectedDestination(const FStringView InDestination
 	WriteSetting(LexToString(BuildType), TEXT("UserSelectedDestination"), InDestination);
 }
 
+void SBuildSelection::SetUserSelectedTransferFilter(const FStringView InIncludeFilter, const FStringView InExcludeFilter)
+{
+	EBuildType BuildType = GetSelectedBuildType();
+	int BuildTypeIndex = (int)BuildType;
+	UserSelectedTransferFilters[BuildTypeIndex].IncludeFilter = InIncludeFilter;
+	UserSelectedTransferFilters[BuildTypeIndex].ExcludeFilter = InExcludeFilter;
+
+	WriteSetting(LexToString(BuildType), TEXT("UserSelectedIncludeFilter"), InIncludeFilter);
+	WriteSetting(LexToString(BuildType), TEXT("UserSelectedExcludeFilter"), InExcludeFilter);
+}
+
 FString SBuildSelection::GetUserSelectedDestination() const
 {
 	EBuildType BuildType = GetSelectedBuildType();
@@ -1362,7 +1382,81 @@ FReply SBuildSelection::ExploreDestination_OnClicked()
 	return FReply::Handled();
 }
 
-void SBuildSelection::OnOpenDestinationDirectoryClicked()
+void SBuildSelection::DownloadSelectedArtifactsInSelectedBuildGroup()
+{
+	using namespace UE::Zen::Build;
+
+	if (TSharedPtr<FBuildServiceInstance> ServiceInstance = BuildServiceInstance.Get())
+	{
+		TArray<FBuildSelectionBuildGroupPtr> SelectedItems = BuildGroupListView->GetSelectedItems();
+		if (SelectedItems.Num() == 0)
+		{
+			return;
+		}
+
+		for (const FString& ArtifactForSelectedGroup : SelectedGroupSelectedArtifacts)
+		{
+			if (FBuildServiceInstance::FBuildRecord* BuildRecord = SelectedItems[0]->NamedArtifacts.Find(ArtifactForSelectedGroup))
+			{
+				FString DestinationPlatformName = ArtifactForSelectedGroup;
+				if (FCbFieldView CookPlatformField = BuildRecord->Metadata["cookPlatform"]; CookPlatformField.HasValue() && !CookPlatformField.HasError())
+				{
+					DestinationPlatformName = *WriteToString<64>(CookPlatformField.AsString());
+				}
+
+				FName TargetPlatformName = NAME_None;
+				if (FCbFieldView PlatformField = BuildRecord->Metadata["platform"]; PlatformField.HasValue() && !PlatformField.HasError())
+				{
+					TargetPlatformName = FName(PlatformField.AsString());
+				}
+
+				EBuildType BuildType = GetSelectedBuildType();
+				int BuildTypeIndex = (int)BuildType;
+				const bool bAppendBuildNameToDestination = bAppendBuildNameToDestinations[BuildTypeIndex];
+
+				FString TransferName = FString::Printf(TEXT("%s-%s"), *SelectedItems[0]->DisplayName, *DestinationPlatformName);
+
+				if (BuildType == EBuildType::Oplog)
+				{
+					FString ProjectFilePath = FUProjectDictionary::GetDefault().GetProjectPathForGame(**SelectedProject);
+					FString DestinationProjectId = GetEffectiveDestination();
+					FString DestinationOplogId = DestinationPlatformName;
+					if (bAppendBuildNameToDestination)
+					{
+						DestinationOplogId = FString::Printf(TEXT("%s.%s"), *DestinationOplogId, *SanitizeForZenId(SelectedItems[0]->DisplayName));
+					}
+					FBuildServiceInstance::FBuildTransfer BuildTransfer =
+						ServiceInstance->StartOplogBuildTransfer(BuildRecord->BuildId, TransferName, DestinationProjectId, DestinationOplogId, ProjectFilePath, SelectedItems[0]->Namespace, BuildRecord->BucketId, TargetPlatformName);
+					OnBuildTransferStarted.ExecuteIfBound(BuildTransfer, SelectedItems[0]->DisplayName, ArtifactForSelectedGroup);
+				}
+				else
+				{
+					FString DestinationFolder;
+					if (bAppendBuildNameToDestination)
+					{
+						DestinationFolder = FPaths::Combine(GetEffectiveDestination(), SanitizeForPath(SelectedItems[0]->DisplayName), DestinationPlatformName);
+					}
+					else
+					{
+						DestinationFolder = FPaths::Combine(GetEffectiveDestination(), DestinationPlatformName);
+					}
+					FBuildServiceInstance::FBuildTransfer BuildTransfer =
+						ServiceInstance->StartBuildTransfer(
+							BuildRecord->BuildId,
+							TransferName,
+							DestinationFolder,
+							SelectedItems[0]->Namespace,
+							BuildRecord->BucketId,
+							UserSelectedTransferFilters[BuildTypeIndex].IncludeFilter,
+							UserSelectedTransferFilters[BuildTypeIndex].ExcludeFilter);
+					OnBuildTransferStarted.ExecuteIfBound(BuildTransfer, SelectedItems[0]->DisplayName, ArtifactForSelectedGroup);
+				}
+			}
+		}
+	}
+}
+
+void SBuildSelection::OnChooseDestinationDirectoryClicked()
 {
 	if (IDesktopPlatform* DesktopPlatform = FDesktopPlatformModule::Get())
 	{
@@ -1380,6 +1474,37 @@ void SBuildSelection::OnOpenDestinationDirectoryClicked()
 	}
 }
 
+void SBuildSelection::OnChooseTransferFilterClicked()
+{
+	if (!FDialogCommands::IsRegistered())
+	{
+		FDialogCommands::Register();
+	}
+
+	EBuildType BuildType = GetSelectedBuildType();
+	int BuildTypeIndex = (int)BuildType;
+
+	TSharedRef<SBuildTransferFilters> TransferFilters = SNew(SBuildTransferFilters)
+		.IncludeFilter(MakeShared<FString>(UserSelectedTransferFilters[BuildTypeIndex].IncludeFilter))
+		.ExcludeFilter(MakeShared<FString>(UserSelectedTransferFilters[BuildTypeIndex].ExcludeFilter));
+
+	TSharedRef<SCustomDialog> TransferFilterDialog = SNew(SCustomDialog)
+		.Title(FText(LOCTEXT("BuildTransferFiltersHeader", "Build Transfer Filter")))
+		.Content()
+		[
+			TransferFilters
+		]
+		.Buttons({
+			SCustomDialog::FButton(LOCTEXT("OK", "OK")),
+			SCustomDialog::FButton(LOCTEXT("Cancel", "Cancel"))
+	});
+
+	if (TransferFilterDialog->ShowModal() == 0)
+	{
+		SetUserSelectedTransferFilter(TransferFilters->GetIncludeFilter(), TransferFilters->GetExcludeFilter());
+	}
+}
+
 TSharedRef<SWidget> SBuildSelection::GetBuildDestinationPanel()
 {
 	TSharedPtr<SCheckBox> AppendBuildNameCheckbox;
@@ -1401,37 +1526,91 @@ TSharedRef<SWidget> SBuildSelection::GetBuildDestinationPanel()
 		.FillWidth(1.0f)
 		.VAlign(VAlign_Center)
 		[
-			SNew(SEditableTextBox)
-			.OverflowPolicy(ETextOverflowPolicy::MiddleEllipsis)
-			.MinDesiredWidth(200.0f)
-			.HintText_Lambda([this]()
-			{
-				return FText::FromString(GetDefaultDestination());
-			})
-			.Text_Lambda([this]()
-			{
-				return FText::FromString(GetUserSelectedDestination());
-			})
-			.OnTextCommitted_Lambda([this](const FText& Text, const ETextCommit::Type CommitType)
-			{
-				SetUserSelectedDestination(Text.ToString());
-			})
+			SNew(SVerticalBox)
+
+			+ SVerticalBox::Slot()
+			.AutoHeight()
+			.VAlign(VAlign_Top)
+			[
+				SNew(SEditableTextBox)
+				.OverflowPolicy(ETextOverflowPolicy::MiddleEllipsis)
+				.MinDesiredWidth(200.0f)
+				.HintText_Lambda([this]()
+				{
+					return FText::FromString(GetDefaultDestination());
+				})
+				.Text_Lambda([this]()
+				{
+					return FText::FromString(GetUserSelectedDestination());
+				})
+				.OnTextCommitted_Lambda([this](const FText& Text, const ETextCommit::Type CommitType)
+				{
+					SetUserSelectedDestination(Text.ToString());
+				})
+			]
+
+			+ SVerticalBox::Slot()
+			.AutoHeight()
+			.VAlign(VAlign_Top)
+			[
+				SNew(SEditableTextBox)
+				.OverflowPolicy(ETextOverflowPolicy::MiddleEllipsis)
+				.MinDesiredWidth(200.0f)
+				.IsReadOnly(true)
+				.Text_Lambda([this]()
+				{
+					EBuildType BuildType = GetSelectedBuildType();
+					int BuildTypeIndex = (int)BuildType;
+					if (!UserSelectedTransferFilters[BuildTypeIndex].ExcludeFilter.IsEmpty())
+					{
+						return FText::Format(LOCTEXT("BuildTransferFilter_IncludeAndExcludeFormat", "Incl: {0}, Excl: {1}"),
+							FText::FromString(UserSelectedTransferFilters[BuildTypeIndex].IncludeFilter),
+							FText::FromString(UserSelectedTransferFilters[BuildTypeIndex].ExcludeFilter));
+					}
+					return FText::FromString(UserSelectedTransferFilters[BuildTypeIndex].IncludeFilter);
+				})
+			]
 		]
 		+ SHorizontalBox::Slot()
 		.AutoWidth()
 		.HAlign(HAlign_Right)
 		[
-			SNew(SButton)
-			.OnClicked_Lambda([this]()
-			{
-				OnOpenDestinationDirectoryClicked();
-				return FReply::Handled();
-			})
-			.ButtonStyle(FAppStyle::Get(), "SimpleButton")
+			SNew(SVerticalBox)
+
+			+ SVerticalBox::Slot()
+			.AutoHeight()
+			.VAlign(VAlign_Top)
 			[
-				SNew(SImage)
-				.Image(FAppStyle::Get().GetBrush("Zen.BrowseContent"))
-				.ColorAndOpacity(FSlateColor::UseForeground())
+				SNew(SButton)
+				.OnClicked_Lambda([this]()
+				{
+					OnChooseDestinationDirectoryClicked();
+					return FReply::Handled();
+				})
+				.ButtonStyle(FAppStyle::Get(), "SimpleButton")
+				[
+					SNew(SImage)
+					.Image(FAppStyle::Get().GetBrush("Zen.BrowseContent"))
+					.ColorAndOpacity(FSlateColor::UseForeground())
+				]
+			]
+
+			+ SVerticalBox::Slot()
+			.AutoHeight()
+			.VAlign(VAlign_Top)
+			[
+				SNew(SButton)
+				.OnClicked_Lambda([this]()
+				{
+					OnChooseTransferFilterClicked();
+					return FReply::Handled();
+				})
+				.ButtonStyle(FAppStyle::Get(), "SimpleButton")
+				[
+					SNew(SImage)
+					.Image(FAppStyle::Get().GetBrush("Zen.Filter"))
+					.ColorAndOpacity(FSlateColor::UseForeground())
+				]
 			]
 		]
 	]
@@ -2039,69 +2218,7 @@ TSharedRef<SWidget> SBuildSelection::GetGridPanel()
 				})
 				.OnClicked_Lambda([this]()
 				{
-					using namespace UE::Zen::Build;
-
-					if (TSharedPtr<FBuildServiceInstance> ServiceInstance = BuildServiceInstance.Get())
-					{
-						TArray<FBuildSelectionBuildGroupPtr> SelectedItems = BuildGroupListView->GetSelectedItems();
-						if (SelectedItems.Num() == 0)
-						{
-							return FReply::Handled();
-						}
-
-						for (const FString& ArtifactForSelectedGroup : SelectedGroupSelectedArtifacts)
-						{
-							if (FBuildServiceInstance::FBuildRecord* BuildRecord = SelectedItems[0]->NamedArtifacts.Find(ArtifactForSelectedGroup))
-							{
-								FString DestinationPlatformName = ArtifactForSelectedGroup;
-								if (FCbFieldView CookPlatformField = BuildRecord->Metadata["cookPlatform"]; CookPlatformField.HasValue() && !CookPlatformField.HasError())
-								{
-									DestinationPlatformName = *WriteToString<64>(CookPlatformField.AsString());
-								}
-
-								FName TargetPlatformName = NAME_None;
-								if (FCbFieldView PlatformField = BuildRecord->Metadata["platform"]; PlatformField.HasValue() && !PlatformField.HasError())
-								{
-									TargetPlatformName = FName(PlatformField.AsString());
-								}
-
-								EBuildType BuildType = GetSelectedBuildType();
-								int BuildTypeIndex = (int)BuildType;
-								const bool bAppendBuildNameToDestination = bAppendBuildNameToDestinations[BuildTypeIndex];
-
-								FString TransferName = FString::Printf(TEXT("%s-%s"), *SelectedItems[0]->DisplayName, *DestinationPlatformName);
-
-								if (BuildType == EBuildType::Oplog)
-								{
-									FString ProjectFilePath = FUProjectDictionary::GetDefault().GetProjectPathForGame(**SelectedProject);
-									FString DestinationProjectId = GetEffectiveDestination();
-									FString DestinationOplogId = DestinationPlatformName;
-									if (bAppendBuildNameToDestination)
-									{
-										DestinationOplogId = FString::Printf(TEXT("%s.%s"), *DestinationOplogId, *SanitizeForZenId(SelectedItems[0]->DisplayName));
-									}
-									FBuildServiceInstance::FBuildTransfer BuildTransfer =
-										ServiceInstance->StartOplogBuildTransfer(BuildRecord->BuildId, TransferName, DestinationProjectId, DestinationOplogId, ProjectFilePath, SelectedItems[0]->Namespace, BuildRecord->BucketId, TargetPlatformName);
-									OnBuildTransferStarted.ExecuteIfBound(BuildTransfer, SelectedItems[0]->DisplayName, ArtifactForSelectedGroup);
-								}
-								else
-								{
-									FString DestinationFolder;
-									if (bAppendBuildNameToDestination)
-									{
-										DestinationFolder = FPaths::Combine(GetEffectiveDestination(), SanitizeForPath(SelectedItems[0]->DisplayName), DestinationPlatformName);
-									}
-									else
-									{
-										DestinationFolder = FPaths::Combine(GetEffectiveDestination(), DestinationPlatformName);
-									}
-									FBuildServiceInstance::FBuildTransfer BuildTransfer =
-										ServiceInstance->StartBuildTransfer(BuildRecord->BuildId, TransferName, DestinationFolder, SelectedItems[0]->Namespace, BuildRecord->BucketId);
-									OnBuildTransferStarted.ExecuteIfBound(BuildTransfer, SelectedItems[0]->DisplayName, ArtifactForSelectedGroup);
-								}
-							}
-						}
-					}
+					DownloadSelectedArtifactsInSelectedBuildGroup();
 					return FReply::Handled();
 				})
 			]
diff --git a/Engine/Source/Editor/StorageServerWidgets/Private/SBuildTransferFilters.cpp b/Engine/Source/Editor/StorageServerWidgets/Private/SBuildTransferFilters.cpp
new file mode 100644
--- /dev/null
+++ b/Engine/Source/Editor/StorageServerWidgets/Private/SBuildTransferFilters.cpp
@@ -0,0 +1,98 @@
+// Copyright Epic Games, Inc. All Rights Reserved.
+
+#include "SBuildTransferFilters.h"
+#include "Styling/StyleColors.h"
+#include "Widgets/Input/SEditableTextBox.h"
+#include "Widgets/Layout/SGridPanel.h"
+#include "Widgets/SBoxPanel.h"
+#include "Widgets/Text/STextBlock.h"
+
+#define LOCTEXT_NAMESPACE "StorageServerBuild"
+
+void SBuildTransferFilters::Construct(const FArguments& InArgs)
+{
+	if (InArgs._IncludeFilter.Get())
+	{
+		IncludeFilter = *InArgs._IncludeFilter.Get();
+	}
+	if (InArgs._ExcludeFilter.Get())
+	{
+		ExcludeFilter = *InArgs._ExcludeFilter.Get();
+	}
+
+	this->ChildSlot
+	[
+		SNew(SVerticalBox)
+
+		+ SVerticalBox::Slot()
+		.AutoHeight()
+		.Padding(0, 0, 0, 0)
+		.Expose(GridSlot)
+		[
+			GetGridPanel()
+		]
+	];
+}
+
+TSharedRef<SWidget> SBuildTransferFilters::GetGridPanel()
+{
+	TSharedRef<SGridPanel> Panel = SNew(SGridPanel)
+		.FillColumn(1, 1.0f);
+
+	const float RowMargin = 0.0f;
+	const float ColumnMargin = 10.0f;
+	const FSlateColor TitleColor = FStyleColors::AccentWhite;
+	const FSlateFontInfo TitleFont = FCoreStyle::GetDefaultFontStyle("Bold", 10);
+
+	Panel->AddSlot(0, 0)
+	[
+		SNew(STextBlock)
+		.Margin(FMargin(ColumnMargin, RowMargin))
+		.ColorAndOpacity(TitleColor)
+		.Font(TitleFont)
+		.Text(LOCTEXT("BuildTransferFilters_IncludeWildcardLabel", "Include Wildcard:"))
+	];
+
+	Panel->AddSlot(1, 0)
+	[
+		SNew(SEditableTextBox)
+		.MinDesiredWidth(200.0f)
+		.Text(FText::FromString(IncludeFilter))
+		.OnTextChanged_Lambda([this](const FText& Text)
+		{
+			IncludeFilter = Text.ToString();
+		})
+		.OnTextCommitted_Lambda([this](const FText& Text, const ETextCommit::Type CommitType)
+		{
+			IncludeFilter = Text.ToString();
+		})
+	];
+
+	Panel->AddSlot(0, 1)
+	[
+		SNew(STextBlock)
+		.Margin(FMargin(ColumnMargin, RowMargin))
+		.ColorAndOpacity(TitleColor)
+		.Font(TitleFont)
+		.Text(LOCTEXT("BuildTransferFilters_ExcludeWildcardLabel", "Exclude Wildcard:"))
+	];
+
+	Panel->AddSlot(1, 1)
+	[
+		SNew(SEditableTextBox)
+		.MinDesiredWidth(200.0f)
+		.Text(FText::FromString(ExcludeFilter))
+		.OnTextChanged_Lambda([this](const FText& Text)
+		{
+			ExcludeFilter = Text.ToString();
+		})
+		.OnTextCommitted_Lambda([this](const FText& Text, const ETextCommit::Type CommitType)
+		{
+			ExcludeFilter = Text.ToString();
+		})
+	];
+
+	return Panel;
+}
+
+#undef LOCTEXT_NAMESPACE
diff --git a/Engine/Source/Editor/StorageServerWidgets/StorageServerWidgets.Build.cs b/Engine/Source/Editor/StorageServerWidgets/StorageServerWidgets.Build.cs
--- a/Engine/Source/Editor/StorageServerWidgets/StorageServerWidgets.Build.cs
+++ b/Engine/Source/Editor/StorageServerWidgets/StorageServerWidgets.Build.cs
@@ -18,6 +18,7 @@ public class StorageServerWidgets : ModuleRules
 				"Projects",	
 				"Slate",
 				"SlateCore",
+				"ToolWidgets",
 				"Zen"
 			});
 
diff --git a/Engine/Source/Programs/BuildStorageTool/Private/BuildStorageToolStyle.cpp b/Engine/Source/Programs/BuildStorageTool/Private/BuildStorageToolStyle.cpp
--- a/Engine/Source/Programs/BuildStorageTool/Private/BuildStorageToolStyle.cpp
+++ b/Engine/Source/Programs/BuildStorageTool/Private/BuildStorageToolStyle.cpp
@@ -63,6 +63,7 @@ TSharedRef< FSlateStyleSet > FBuildStorageToolStyle::Create()
 	Style.Set("Zen.FolderExplore", new EDITOR_IMAGE_BRUSH_SVG("Starship/Common/ContentBrowser", Icon12x12));
 	Style.Set("Zen.BrowseContent", new EDITOR_IMAGE_BRUSH_SVG("Starship/Common/BrowseContent", Icon16x16));
 	Style.Set("Zen.Clipboard", new EDITOR_IMAGE_BRUSH_SVG("Starship/Common/Clipboard", Icon16x16));
+	Style.Set("Zen.Filter", new CORE_IMAGE_BRUSH_SVG("Starship/Common/filter", Icon16x16));
 
 	Style.Set("Zen.FolderView", new CORE_IMAGE_BRUSH_SVG("Starship/Common/folder-open", Icon12x12));
 	Style.Set("Zen.BrowserView", new EDITOR_IMAGE_BRUSH_SVG("Starship/Common/WebBrowser", Icon12x12));
