BuildStorageTool: Improve the way we launch the folder browser when the destination folder is not created yet.

diff --git a/Engine/Source/Editor/StorageServerWidgets/Private/SBuildActivity.cpp b/Engine/Source/Editor/StorageServerWidgets/Private/SBuildActivity.cpp
--- a/Engine/Source/Editor/StorageServerWidgets/Private/SBuildActivity.cpp
+++ b/Engine/Source/Editor/StorageServerWidgets/Private/SBuildActivity.cpp
@@ -305,14 +305,18 @@ void SBuildActivity::OpenDestinationForBuildActivity(TSharedPtr<UE::Zen::Build::
 			break;
 		case FBuildServiceInstance::FBuildTransfer::EType::Files:
 			{
-				FString DestinationFolder = Item->Transfer.GetDestination();
-				if (FPaths::DirectoryExists(DestinationFolder))
+				FString FolderToExplore = Item->Transfer.GetDestination();
+				for (int Index = 0; Index < 3; ++Index)
 				{
-					FPlatformProcess::ExploreFolder(*DestinationFolder);
+					if (FolderToExplore.IsEmpty() || FPaths::DirectoryExists(FolderToExplore))
+					{
+						break;
+					}
+					FolderToExplore = FPaths::GetPath(FolderToExplore);
 				}
-				else
+				if (!FolderToExplore.IsEmpty() && FPaths::DirectoryExists(FolderToExplore))
 				{
-					FPlatformProcess::ExploreFolder(*FPaths::GetPath(DestinationFolder));
+					FPlatformProcess::ExploreFolder(*FolderToExplore);
 				}
 			}
 			break;
