[StorageServerWidgets]
* Fixed header unit compile error

diff --git a/Engine/Source/Editor/StorageServerWidgets/Private/SBuildSelection.cpp b/Engine/Source/Editor/StorageServerWidgets/Private/SBuildSelection.cpp
--- a/Engine/Source/Editor/StorageServerWidgets/Private/SBuildSelection.cpp
+++ b/Engine/Source/Editor/StorageServerWidgets/Private/SBuildSelection.cpp
@@ -1181,7 +1181,7 @@ FString SBuildSelection::GetDefaultDestination() const
 		}
 
 		return PlatformFile.ConvertToAbsolutePathForExternalAppForRead(
-				*FPaths::FPaths::Combine(FPaths::EngineSavedDir(), TEXT("DownloadedBuilds"))
+				*FPaths::Combine(FPaths::EngineSavedDir(), TEXT("DownloadedBuilds"))
 			);
 	}
 
