Fix parsing of ProjectFile commandline argument from build selection dialog.

diff --git a/Engine/Source/Editor/StorageServerWidgets/Private/SBuildSelection.cpp b/Engine/Source/Editor/StorageServerWidgets/Private/SBuildSelection.cpp
--- a/Engine/Source/Editor/StorageServerWidgets/Private/SBuildSelection.cpp
+++ b/Engine/Source/Editor/StorageServerWidgets/Private/SBuildSelection.cpp
@@ -79,6 +79,8 @@ void SBuildSelection::Construct(const FArguments& InArgs)
 	FString ProjectFile;
 	if (FParse::Value(FCommandLine::Get(), TEXT("-ProjectFile"), ProjectFile))
 	{
+		ProjectFile.TrimStartAndEndInline();
+		ProjectFile.TrimQuotesInline();
 		if (IDesktopPlatform* DesktopPlatform = FDesktopPlatformModule::Get())
 		{
 			for (const TPair<FString, FString>& EngineInstallation : EngineInstallations)
