Improve UI feedback in BST progress field

#rnx

diff --git a/Engine/Source/Developer/Zen/Private/BuildServerInterface.cpp b/Engine/Source/Developer/Zen/Private/BuildServerInterface.cpp
--- a/Engine/Source/Developer/Zen/Private/BuildServerInterface.cpp
+++ b/Engine/Source/Developer/Zen/Private/BuildServerInterface.cpp
@@ -657,7 +657,11 @@ FBuildServiceInstance::FBuildTransfer::GetCurrentProgress(FString& OutLabel, FSt
 		{
 			OutLabel = State->RecentOutputLines[State->RecentOutputLines.GetPreviousIndex(State->NextOutputLineIndex)];
 		}
-		OutDetail = State->GetRecentOutput();
+		if (State->NextOutputLineIndex > 0)
+		{
+			OutDetail = State->RecentOutputLines[State->NextOutputLineIndex - 1];
+		}
+
 		OutPercent = 0.f;
 		return true;
 	}
@@ -686,7 +690,11 @@ FBuildServiceInstance::FBuildTransfer::GetOverallProgress(FString& OutLabel, FSt
 		{
 			OutLabel = State->RecentOutputLines[State->RecentOutputLines.GetPreviousIndex(State->NextOutputLineIndex)];
 		}
-		OutDetail = State->GetRecentOutput();
+		if (State->NextOutputLineIndex > 0)
+		{
+			OutDetail = State->RecentOutputLines[State->NextOutputLineIndex - 1];
+		}
+
 		OutPercent = 0.f;
 		return true;
 	}
diff --git a/Engine/Source/Editor/StorageServerWidgets/Private/SBuildActivity.cpp b/Engine/Source/Editor/StorageServerWidgets/Private/SBuildActivity.cpp
--- a/Engine/Source/Editor/StorageServerWidgets/Private/SBuildActivity.cpp
+++ b/Engine/Source/Editor/StorageServerWidgets/Private/SBuildActivity.cpp
@@ -533,22 +533,16 @@ SBuildActivityTableRow::GenerateWidgetForColumn(const FName& ColumnName)
 							float Percent;
 							if (BuildActivity->Transfer.GetCurrentProgress(Label, Detail, Percent))
 							{
-								return FText::FromString(Label);
+								return FText::FromString(Detail.IsEmpty() ? Label : Detail);
 							}
 							return FText::FromString(FString());
 						})
 						.TextStyle(&FAppStyle::GetWidgetStyle<FTextBlockStyle>(TEXT("SmallText")))
 						.Justification(ETextJustify::Type::Left)
+						.WrapTextAt(0.f)
 						.ToolTipText_Lambda([this]()
 						{
-							FString Label;
-							FString Detail;
-							float Percent;
-							if (BuildActivity->Transfer.GetCurrentProgress(Label, Detail, Percent))
-							{
-								return FText::FromString(Detail);
-							}
-							return FText::FromString(FString());
+							return FText::FromString(BuildActivity->Transfer.GetRecentOutput());
 						})
 						.Visibility_Lambda([this]()
 						{
