Fix UI bug in BuildStorageTool where the "Append build name" checkbox would sometimes show the wrong state.

diff --git a/Engine/Source/Editor/StorageServerWidgets/Private/SBuildSelection.cpp b/Engine/Source/Editor/StorageServerWidgets/Private/SBuildSelection.cpp
--- a/Engine/Source/Editor/StorageServerWidgets/Private/SBuildSelection.cpp
+++ b/Engine/Source/Editor/StorageServerWidgets/Private/SBuildSelection.cpp
@@ -1722,7 +1722,7 @@ TSharedRef<SWidget> SBuildSelection::GetBuildDestinationPanel()
 					{
 						NewState = ECheckBoxState::Unchecked;
 					}
-					AppendBuildNameCheckbox->SetIsChecked(NewState);
+
 					if (NewState == ECheckBoxState::Checked)
 					{
 						bAppendBuildNameToDestinations[BuildTypeIndex] = true;
