Fix case-sensitive metadata search within build storage tool.

diff --git a/Engine/Source/Editor/StorageServerWidgets/Private/SBuildSelection.cpp b/Engine/Source/Editor/StorageServerWidgets/Private/SBuildSelection.cpp
--- a/Engine/Source/Editor/StorageServerWidgets/Private/SBuildSelection.cpp
+++ b/Engine/Source/Editor/StorageServerWidgets/Private/SBuildSelection.cpp
@@ -494,11 +494,11 @@ void SBuildSelection::RegenerateBuildGroups(UE::Zen::Build::FListBuildsState& Li
 
 	auto MakeGroupKey = [](const FString& Namespace, const FString& CommitIdentifier, const FBuildServiceInstance::FBuildRecord& BuildRecord)
 	{
-		if (FCbFieldView BuildGroupField = BuildRecord.Metadata["buildgroup"]; BuildGroupField.HasValue() && !BuildGroupField.HasError())
+		if (FCbFieldView BuildGroupField = BuildRecord.Metadata.FindIgnoreCase("buildgroup"); BuildGroupField.HasValue() && !BuildGroupField.HasError())
 		{
 			return FString(WriteToString<64>(Namespace, ".", BuildGroupField.AsString()));
 		}
-		if (FCbFieldView NameField = BuildRecord.Metadata["name"]; NameField.HasValue() && !NameField.HasError())
+		if (FCbFieldView NameField = BuildRecord.Metadata.FindIgnoreCase("name"); NameField.HasValue() && !NameField.HasError())
 		{
 			FString CandidateKeyName = FString(WriteToString<64>(Namespace, ".", NameField.AsString()));
 			// TODO: This name manipulation needs to be removed when the metadata is more consistent.
@@ -543,7 +543,7 @@ void SBuildSelection::RegenerateBuildGroups(UE::Zen::Build::FListBuildsState& Li
 			}
 
 			FString ArtifactName;
-			if (FCbFieldView NameView = BuildRecord.Metadata["name"]; NameView.HasValue() && !NameView.HasError())
+			if (FCbFieldView NameView = BuildRecord.Metadata.FindIgnoreCase("name"); NameView.HasValue() && !NameView.HasError())
 			{
 				ArtifactName = FUTF8ToTCHAR(NameView.AsString());
 
@@ -577,19 +577,19 @@ void SBuildSelection::RegenerateBuildGroups(UE::Zen::Build::FListBuildsState& Li
 			if (BuildGroup->DisplayName.IsEmpty())
 			{
 				FString Job;
-				if (FCbFieldView JobField = BuildRecord.Metadata["job"]; JobField.HasValue() && !JobField.HasError())
+				if (FCbFieldView JobField = BuildRecord.Metadata.FindIgnoreCase("job"); JobField.HasValue() && !JobField.HasError())
 				{
 					Job = FUTF8ToTCHAR(JobField.AsString());
 				}
 
 				FString Category;
-				if (FCbFieldView TemplateIdField = BuildRecord.Metadata["hordeTemplateId"]; TemplateIdField.HasValue() && !TemplateIdField.HasError())
+				if (FCbFieldView TemplateIdField = BuildRecord.Metadata.FindIgnoreCase("hordeTemplateId"); TemplateIdField.HasValue() && !TemplateIdField.HasError())
 				{
 					Category = FUTF8ToTCHAR(TemplateIdField.AsString());
 				}
 
 				bool bIsPreflight = false;
-				if (FCbFieldView IsPreflightField = BuildRecord.Metadata["ispreflight"]; IsPreflightField.HasValue() && !IsPreflightField.HasError())
+				if (FCbFieldView IsPreflightField = BuildRecord.Metadata.FindIgnoreCase("ispreflight"); IsPreflightField.HasValue() && !IsPreflightField.HasError())
 				{
 					FString IsPreflightString;
 					IsPreflightString = FUTF8ToTCHAR(IsPreflightField.AsString());
@@ -603,7 +603,7 @@ void SBuildSelection::RegenerateBuildGroups(UE::Zen::Build::FListBuildsState& Li
 				}
 
 				FDateTime CreatedAt;
-				if (FCbFieldView CreatedAtField = BuildRecord.Metadata["createdAt"]; CreatedAtField.HasValue() && !CreatedAtField.HasError())
+				if (FCbFieldView CreatedAtField = BuildRecord.Metadata.FindIgnoreCase("createdAt"); CreatedAtField.HasValue() && !CreatedAtField.HasError())
 				{
 					if (CreatedAtField.IsString())
 					{
@@ -616,16 +616,16 @@ void SBuildSelection::RegenerateBuildGroups(UE::Zen::Build::FListBuildsState& Li
 				}
 
 				FString Configuration;
-				if (FCbFieldView ConfigurationField = BuildRecord.Metadata["Configuration"]; ConfigurationField.HasValue() && !ConfigurationField.HasError())
+				if (FCbFieldView ConfigurationField = BuildRecord.Metadata.FindIgnoreCase("Configuration"); ConfigurationField.HasValue() && !ConfigurationField.HasError())
 				{
 					Configuration = FUTF8ToTCHAR(ConfigurationField.AsString());
 				}
 
 				FString ItemName;
-				FCbFieldView GroupNameView = BuildRecord.Metadata["buildgroup"];
+				FCbFieldView GroupNameView = BuildRecord.Metadata.FindIgnoreCase("buildgroup");
 				if (!GroupNameView.HasValue())
 				{
-					GroupNameView = BuildRecord.Metadata["name"];
+					GroupNameView = BuildRecord.Metadata.FindIgnoreCase("name");
 				}
 				if (FCbFieldView NameView = GroupNameView; NameView.HasValue() && !NameView.HasError())
 				{
