When listing builds in BuildStorageTool eliminate "Suffix" as a distinct concept and just use the supplied BuildGroup as-is without manipulation.  Anything in the suffix should be evident in the original BuildGroup.

diff --git a/Engine/Source/Editor/StorageServerWidgets/Internal/SBuildSelection.h b/Engine/Source/Editor/StorageServerWidgets/Internal/SBuildSelection.h
--- a/Engine/Source/Editor/StorageServerWidgets/Internal/SBuildSelection.h
+++ b/Engine/Source/Editor/StorageServerWidgets/Internal/SBuildSelection.h
@@ -50,7 +50,6 @@ class SBuildSelection : public SCompoundWidget
 		FString Namespace;
 		FString DisplayName;
 		FString CommitIdentifier;
-		FString Suffix;
 		FDateTime CreatedAt;
 		FString Category;
 		FString Job;
diff --git a/Engine/Source/Editor/StorageServerWidgets/Private/SBuildSelection.cpp b/Engine/Source/Editor/StorageServerWidgets/Private/SBuildSelection.cpp
--- a/Engine/Source/Editor/StorageServerWidgets/Private/SBuildSelection.cpp
+++ b/Engine/Source/Editor/StorageServerWidgets/Private/SBuildSelection.cpp
@@ -44,7 +44,6 @@ namespace UE::BuildSelection::Internal
 	{
 		const FName ColName = TEXT("Name");
 		const FName ColCommit = TEXT("Commit");
-		const FName ColSuffix = TEXT("Suffix");
 		const FName ColCategory = TEXT("Category");
 		const FName ColCreated = TEXT("Created");
 	}
@@ -596,12 +595,6 @@ void SBuildSelection::RegenerateBuildGroups(UE::Zen::Build::FListBuildsState& Li
 					bIsPreflight = IsPreflightString.ToBool();
 				}
 
-				FString Suffix("");
-				if (bIsPreflight)
-				{
-					Suffix.Append("PF");
-				}
-
 				FDateTime CreatedAt;
 				if (FCbFieldView CreatedAtField = BuildRecord.Metadata.FindIgnoreCase("createdAt"); CreatedAtField.HasValue() && !CreatedAtField.HasError())
 				{
@@ -631,39 +624,6 @@ void SBuildSelection::RegenerateBuildGroups(UE::Zen::Build::FListBuildsState& Li
 				{
 					// TODO: This name manipulation needs to be removed when the metadata is more consistent.
 					ItemName = FUTF8ToTCHAR(NameView.AsString());
-					int32 CLStartIndex = ItemName.Find(TEXT("-CL"));
-					if (CLStartIndex != INDEX_NONE)
-					{
-						int32 TruncationIndex = ItemName.Find(TEXT("-"), ESearchCase::IgnoreCase, ESearchDir::FromStart, CLStartIndex + 4);
-						if (TruncationIndex != INDEX_NONE)
-						{
-							ItemName.LeftInline(TruncationIndex);
-						}
-						TruncationIndex = ItemName.Find(TEXT("."), ESearchCase::IgnoreCase, ESearchDir::FromStart, CLStartIndex + 4);
-						if (TruncationIndex != INDEX_NONE)
-						{
-							ItemName.LeftInline(TruncationIndex);
-						}
-					}
-
-					if (!Category.IsEmpty() && ItemName.StartsWith(Category))
-					{
-						ItemName.RightChopInline(Category.Len());
-					}
-
-					bool bCharRemoved = false;
-					do
-					{
-						ItemName.TrimCharInline(TCHAR('.'), &bCharRemoved);
-					}
-					while(bCharRemoved);
-					do
-					{
-						ItemName.TrimCharInline(TCHAR('+'), &bCharRemoved);
-					}
-					while(bCharRemoved);
-
-					ItemName.ReplaceCharInline(TCHAR('+'), TCHAR('-'));
 				}
 				else
 				{
@@ -673,7 +633,6 @@ void SBuildSelection::RegenerateBuildGroups(UE::Zen::Build::FListBuildsState& Li
 				BuildGroup->Namespace = BuildState.Namespace;
 				BuildGroup->DisplayName = ItemName;
 				BuildGroup->CommitIdentifier = CommitIdentifier;
-				BuildGroup->Suffix = Suffix;
 				BuildGroup->Category = Category;
 				BuildGroup->CreatedAt = CreatedAt;
 				BuildGroup->Job = Job;
@@ -713,12 +672,6 @@ void SBuildSelection::ConditionalSortBuildGroups()
 				return (A->CommitIdentifier < B->CommitIdentifier);
 			});
 		}
-		else if (BuildGroupSortByColumn == FBuildGroupIds::ColSuffix)
-		{
-			BuildGroups.Sort([this](const TSharedPtr<FBuildGroup>& A, const TSharedPtr<FBuildGroup>& B) {
-				return (A->Suffix < B->Suffix);
-			});
-		}
 		else if (BuildGroupSortByColumn == FBuildGroupIds::ColCategory)
 		{
 			BuildGroups.Sort([this](const TSharedPtr<FBuildGroup>& A, const TSharedPtr<FBuildGroup>& B) {
@@ -746,12 +699,6 @@ void SBuildSelection::ConditionalSortBuildGroups()
 				return (A->CommitIdentifier >= B->CommitIdentifier);
 			});
 		}
-		else if (BuildGroupSortByColumn == FBuildGroupIds::ColSuffix)
-		{
-			BuildGroups.Sort([this](const TSharedPtr<FBuildGroup>& A, const TSharedPtr<FBuildGroup>& B) {
-				return (A->Suffix >= B->Suffix);
-			});
-		}
 		else if (BuildGroupSortByColumn == FBuildGroupIds::ColCategory)
 		{
 			BuildGroups.Sort([this](const TSharedPtr<FBuildGroup>& A, const TSharedPtr<FBuildGroup>& B) {
@@ -850,7 +797,6 @@ bool SBuildSelection::DoesPassAllViewFilters(TSharedPtr<FBuildGroup> InItem) con
 	if (InItem->Namespace.Contains(SelectedTextFilter) ||
 		InItem->DisplayName.Contains(SelectedTextFilter) ||
 		InItem->CommitIdentifier.Contains(SelectedTextFilter) ||
-		InItem->Suffix.Contains(SelectedTextFilter) ||
 		InItem->Category.Contains(SelectedTextFilter) ||
 		InItem->Job.Contains(SelectedTextFilter))
 	{
@@ -2104,11 +2050,6 @@ TSharedRef<SWidget> SBuildSelection::GetGridPanel()
 						.FillWidth(0.10f).HAlignCell(HAlign_Center).HAlignHeader(HAlign_Center).VAlignCell(VAlign_Center)
 						.SortMode(this, &SBuildSelection::GetColumnSortMode, FBuildGroupIds::ColCommit)
 						.OnSort(this, &SBuildSelection::OnColumnSortModeChanged)
-					+ SHeaderRow::Column(FBuildGroupIds::ColSuffix).DefaultLabel(LOCTEXT("BuildSelection_BuildGroupColSuffix", "Suffix"))
-						.DefaultTooltip(LOCTEXT("BuildSelection_BuildGroupColSuffixTooltip", "Modifier on top of the commit/changelist for the build"))
-						.FillWidth(0.10f).HAlignCell(HAlign_Center).HAlignHeader(HAlign_Center).VAlignCell(VAlign_Center)
-						.SortMode(this, &SBuildSelection::GetColumnSortMode, FBuildGroupIds::ColSuffix)
-						.OnSort(this, &SBuildSelection::OnColumnSortModeChanged)
 					+ SHeaderRow::Column(FBuildGroupIds::ColCategory).DefaultLabel(LOCTEXT("BuildSelection_BuildGroupColCategory", "Category"))
 						.DefaultTooltip(LOCTEXT("BuildSelection_BuildGroupColCategoryTooltip", "Category for the build"))
 						.SortMode(this, &SBuildSelection::GetColumnSortMode, FBuildGroupIds::ColCategory)
@@ -2304,15 +2245,6 @@ SBuildGroupTableRow::GenerateWidgetForColumn(const FName& ColumnName)
 					SNew(STextBlock).Text(FText::FromString(BuildGroup->CommitIdentifier))
 				];
 	}
-	else if (ColumnName == FBuildGroupIds::ColSuffix)
-	{
-		return SNew(SHorizontalBox)
-				+ SHorizontalBox::Slot()
-				.Padding(1.f,1.f)
-				[
-					SNew(STextBlock).Text(FText::FromString(BuildGroup->Suffix))
-				];
-	}
 	else if (ColumnName == FBuildGroupIds::ColCategory)
 	{
 		return SNew(SHorizontalBox)
