When choosing a destination oplog in BuildStorageTool, trim leading underscores of the sanitized path while ensuring that we don't trim to an empty string.

diff --git a/Engine/Source/Editor/StorageServerWidgets/Private/SBuildSelection.cpp b/Engine/Source/Editor/StorageServerWidgets/Private/SBuildSelection.cpp
--- a/Engine/Source/Editor/StorageServerWidgets/Private/SBuildSelection.cpp
+++ b/Engine/Source/Editor/StorageServerWidgets/Private/SBuildSelection.cpp
@@ -1274,7 +1274,16 @@ FString SBuildSelection::SanitizeForZenId(const FString& InString)
 			OutString[Index] = TCHAR('_');
 		}
 	}
-	return OutString;
+
+	// Trim leading and trailing underscores
+	FString OutStringPreTrim = OutString;
+	bool bCharsWereRemoved;
+	do
+	{
+		OutString.TrimCharInline(TCHAR('_'), &bCharsWereRemoved);
+	} while (bCharsWereRemoved);
+
+	return OutString.IsEmpty() ? OutStringPreTrim : OutString;
 }
 
 FString SBuildSelection::SanitizeBucketSegment(const FString& InString)
