BuildStorageTool: Make sure that we invalidate listview geometry and request a refresh after changes that filter the listview contents.

diff --git a/Engine/Source/Editor/StorageServerWidgets/Internal/SBuildSelection.h b/Engine/Source/Editor/StorageServerWidgets/Internal/SBuildSelection.h
--- a/Engine/Source/Editor/StorageServerWidgets/Internal/SBuildSelection.h
+++ b/Engine/Source/Editor/StorageServerWidgets/Internal/SBuildSelection.h
@@ -86,6 +86,7 @@ private:
 	TSharedRef<SWidget> OnGenerateTextBlockFromString(TSharedPtr<FString> Item);
 	TSharedRef<SWidget> OnGenerateBuildTypeTextBlockFromString(TSharedPtr<FString> Item);
 	bool DoesPassAllViewFilters(TSharedPtr<FBuildGroup> InItem) const;
+	void PostViewFilterChange();
 	bool BuildGroupIsSelectableOrNavigable(TSharedPtr<FBuildGroup> InItem) const;
 	TSharedRef<ITableRow> GenerateBuildGroupRow(TSharedPtr<FBuildGroup> InItem, const TSharedRef<STableViewBase>& InOwningTable);
 	TSharedPtr<SWidget> OnGetBuildGroupContextMenuContent() const;
diff --git a/Engine/Source/Editor/StorageServerWidgets/Private/SBuildSelection.cpp b/Engine/Source/Editor/StorageServerWidgets/Private/SBuildSelection.cpp
--- a/Engine/Source/Editor/StorageServerWidgets/Private/SBuildSelection.cpp
+++ b/Engine/Source/Editor/StorageServerWidgets/Private/SBuildSelection.cpp
@@ -454,7 +454,7 @@ void SBuildSelection::RebuildLists()
 							ExecuteOnGameThread(UE_SOURCE_LOCATION,
 								[this]
 								{
-									BuildGroupListView->RequestListRefresh();
+									PostViewFilterChange();
 									BuildListRefreshesInProgress.fetch_sub(1);
 								});
 						}
@@ -811,6 +811,13 @@ bool SBuildSelection::DoesPassAllViewFilters(TSharedPtr<FBuildGroup> InItem) con
 	return false;
 }
 
+void SBuildSelection::PostViewFilterChange()
+{
+	ValidateBuildGroupSelection();
+	BuildGroupListView->Invalidate(EInvalidateWidgetReason::LayoutAndVolatility);
+	BuildGroupListView->RequestListRefresh();
+}
+
 bool SBuildSelection::BuildGroupIsSelectableOrNavigable(FBuildSelectionBuildGroupPtr InItem) const
 {
 	return DoesPassAllViewFilters(InItem);
@@ -1028,7 +1035,7 @@ void SBuildSelection::OnColumnSortModeChanged(const EColumnSortPriority::Type So
 	ExecuteOnGameThread(UE_SOURCE_LOCATION,
 		[this]
 		{
-			BuildGroupListView->RequestListRefresh();
+			PostViewFilterChange();
 		});
 }
 
@@ -1195,6 +1202,7 @@ void SBuildSelection::SetIncludePreflights(const bool bInIncludePreflights)
 {
 	bIncludePreflights = bInIncludePreflights;
 	WriteSetting(TEXT("IncludePreflights"), bIncludePreflights ? TEXT("True") : TEXT("False"));
+	PostViewFilterChange();
 }
 
 bool SBuildSelection::WriteSetting(const FStringView InSectionName, const FStringView InKeyName, const FStringView InValue)
@@ -1925,10 +1933,12 @@ TSharedRef<SWidget> SBuildSelection::GetGridPanel()
 				.OnTextChanged_Lambda([this](const FText& Text)
 				{
 					SelectedTextFilter = Text.ToString();
+					PostViewFilterChange();
 				})
 				.OnTextCommitted_Lambda([this](const FText& Text, const ETextCommit::Type CommitType)
 				{
 					SelectedTextFilter = Text.ToString();
+					PostViewFilterChange();
 				})
 			]
 		]
@@ -1964,7 +1974,7 @@ TSharedRef<SWidget> SBuildSelection::GetGridPanel()
 				.OnCheckedValuesChanged_Lambda([this]()
 				{
 					RegenerateActivePlatformFilters();
-					ValidateBuildGroupSelection();
+					PostViewFilterChange();
 				})
 			]
 		]
