BuildStorageTool: change the heuristic for selecting row count for available builds to allow it to increase the row count faster as the build count increases.  Went from sqrt(x3 (ie: x^0.5 + 2.0) to pow(x, 0.7), swapping a linear offset for a change in exponent to allow faster growth.

diff --git a/Engine/Source/Editor/StorageServerWidgets/Private/SBuildSelection.cpp b/Engine/Source/Editor/StorageServerWidgets/Private/SBuildSelection.cpp
--- a/Engine/Source/Editor/StorageServerWidgets/Private/SBuildSelection.cpp
+++ b/Engine/Source/Editor/StorageServerWidgets/Private/SBuildSelection.cpp
@@ -904,7 +904,7 @@ void SBuildSelection::BuildGroupSelectionChanged(FBuildSelectionBuildGroupPtr It
 	int32 Column = 0;
 	const int32 TotalListCount = ArtifactNameSet.Num();
 	// Heuristic for choosing grid rows when laying out rows and columns for the quantity of platforms
-	const int32 MaxRows = FMath::Max(4, FMath::FloorToInt32(FMath::Sqrt(float(TotalListCount))) + 2);
+	const int32 MaxRows = FMath::Max(4, FMath::FloorToInt32(FMath::Pow(float(TotalListCount), 0.7)));
 
 	for (const FString& ArtifactName : ArtifactNameSet)
 	{
