BuildStorageTool: UX improvements related to the transfer filtering interface.

diff --git a/Engine/Source/Editor/StorageServerWidgets/Private/SBuildSelection.cpp b/Engine/Source/Editor/StorageServerWidgets/Private/SBuildSelection.cpp
--- a/Engine/Source/Editor/StorageServerWidgets/Private/SBuildSelection.cpp
+++ b/Engine/Source/Editor/StorageServerWidgets/Private/SBuildSelection.cpp
@@ -1577,7 +1577,11 @@ TSharedRef<SWidget> SBuildSelection::GetBuildDestinationPanel()
 							FText::FromString(UserSelectedTransferFilters[BuildTypeIndex].IncludeFilter),
 							FText::FromString(UserSelectedTransferFilters[BuildTypeIndex].ExcludeFilter));
 					}
-					return FText::FromString(UserSelectedTransferFilters[BuildTypeIndex].IncludeFilter);
+					if (!UserSelectedTransferFilters[BuildTypeIndex].IncludeFilter.IsEmpty())
+					{
+						return FText::FromString(UserSelectedTransferFilters[BuildTypeIndex].IncludeFilter);
+					}
+					return LOCTEXT("BuildTransferFilter_DownloadAllFiles", "<All files>");
 				})
 			]
 		]
@@ -1610,6 +1614,7 @@ TSharedRef<SWidget> SBuildSelection::GetBuildDestinationPanel()
 			.VAlign(VAlign_Top)
 			[
 				SNew(SButton)
+				.ToolTipText(LOCTEXT("BuildSelection_ChooseTransferFilterTooltip", "Change what gets downloaded within builds."))
 				.OnClicked_Lambda([this]()
 				{
 					OnChooseTransferFilterClicked();
diff --git a/Engine/Source/Editor/StorageServerWidgets/Private/SBuildTransferFilters.cpp b/Engine/Source/Editor/StorageServerWidgets/Private/SBuildTransferFilters.cpp
--- a/Engine/Source/Editor/StorageServerWidgets/Private/SBuildTransferFilters.cpp
+++ b/Engine/Source/Editor/StorageServerWidgets/Private/SBuildTransferFilters.cpp
@@ -50,13 +50,13 @@ TSharedRef<SWidget> SBuildTransferFilters::GetGridPanel()
 		.Margin(FMargin(ColumnMargin, RowMargin))
 		.ColorAndOpacity(TitleColor)
 		.Font(TitleFont)
-		.Text(LOCTEXT("BuildTransferFilters_IncludeWildcardLabel", "Include Wildcard:"))
+		.Text(LOCTEXT("BuildTransferFilters_IncludeFilterLabel", "Include Filter:"))
 	];
 
 	Panel->AddSlot(1, 0)
 	[
 		SNew(SEditableTextBox)
-		.MinDesiredWidth(200.0f)
+		.MinDesiredWidth(400.0f)
 		.Text(FText::FromString(IncludeFilter))
 		.OnTextChanged_Lambda([this](const FText& Text)
 		{
@@ -74,13 +74,13 @@ TSharedRef<SWidget> SBuildTransferFilters::GetGridPanel()
 		.Margin(FMargin(ColumnMargin, RowMargin))
 		.ColorAndOpacity(TitleColor)
 		.Font(TitleFont)
-		.Text(LOCTEXT("BuildTransferFilters_ExcludeWildcardLabel", "Exclude Wildcard:"))
+		.Text(LOCTEXT("BuildTransferFilters_ExcludeFilterLabel", "Exclude Filter:"))
 	];
 
 	Panel->AddSlot(1, 1)
 	[
 		SNew(SEditableTextBox)
-		.MinDesiredWidth(200.0f)
+		.MinDesiredWidth(400.0f)
 		.Text(FText::FromString(ExcludeFilter))
 		.OnTextChanged_Lambda([this](const FText& Text)
 		{
@@ -92,6 +92,29 @@ TSharedRef<SWidget> SBuildTransferFilters::GetGridPanel()
 		})
 	];
 
+	Panel->AddSlot(0, 2)
+	[
+		SNew(STextBlock)
+		.Margin(FMargin(ColumnMargin, 10.0f))
+		.ColorAndOpacity(TitleColor)
+		.Font(TitleFont)
+		.Text(LOCTEXT("BuildTransferFilters_NotesLabel", "Notes:"))
+	];
+
+	Panel->AddSlot(1, 2)
+	[
+		SNew(STextBlock)
+		.Margin(FMargin(0.0f, 10.0f))
+		.Text(LOCTEXT("BuildTransferFilters_NotesContents",
+			"Windows style wildcard(s) (using * and ?) to match file paths to exclude, separated by semicolon(;).\n"
+			"Specifying a non-empty inclusion limits the download to only files that match the inclusion patterns.\n"
+			"Exclusions override any inclusions.\n\n"
+			"Examples:\n"
+			"Include filter to download Windows executables only: *.exe;*.dll\n"
+			"Exclude filter to not download symbol files: *.pdb*;*.map"
+			))
+	];
+
 	return Panel;
 }
 
