BuildStorageTool: UX and misc cleanup and bug fixes:
- Fix the way we consume/use branch from the commandline to address a flaw in the way dot characters are handled
- Add a copy name menu item to the right click context menu for build groups
- Cleanup/remove some unused variables

diff --git a/Engine/Source/Editor/StorageServerWidgets/Internal/SBuildSelection.h b/Engine/Source/Editor/StorageServerWidgets/Internal/SBuildSelection.h
--- a/Engine/Source/Editor/StorageServerWidgets/Internal/SBuildSelection.h
+++ b/Engine/Source/Editor/StorageServerWidgets/Internal/SBuildSelection.h
@@ -123,7 +123,8 @@ private:
 
 	static FString SanitizeForPath(const FString& InString);
 	static FString SanitizeForZenId(const FString& InString);
-	static FString SanitizeBucketSegment(const FString& InString);
+	static FString SanitizeBranch(const FString& InString);
+	static FString SanitizeBucket(const FString& InString);
 
 	static const TCHAR* LexToString(EBuildType BuildType);
 
@@ -155,8 +156,6 @@ private:
 	TSharedPtr<SGridPanel> SelectedGroupArtifactGrid;
 	TArray<FString> SelectedGroupSelectedArtifacts;
 
-	TMultiMap<FString, FString> BucketsToNamespaces;
-
 	TArray<TSharedPtr<FString>> ProjectList;
 	TArray<TSharedPtr<FString>> StreamList;
 	TArray<TSharedPtr<FString>> BuildTypeList;
@@ -191,7 +190,6 @@ private:
 	std::atomic<uint32> BuildRefreshGeneration = 0;
 	TArray<TSharedPtr<SBuildSelection::FBuildGroup, ESPMode::ThreadSafe>> BuildGroups;
 
-	SVerticalBox::FSlot* GridSlot = nullptr;
 	TAttribute<TSharedPtr<UE::Zen::FZenServiceInstance>> ZenServiceInstance;
 	TAttribute<TSharedPtr<UE::Zen::Build::FBuildServiceInstance>> BuildServiceInstance;
 	FOnBuildTransferStarted OnBuildTransferStarted;
diff --git a/Engine/Source/Editor/StorageServerWidgets/Private/SBuildSelection.cpp b/Engine/Source/Editor/StorageServerWidgets/Private/SBuildSelection.cpp
--- a/Engine/Source/Editor/StorageServerWidgets/Private/SBuildSelection.cpp
+++ b/Engine/Source/Editor/StorageServerWidgets/Private/SBuildSelection.cpp
@@ -134,7 +134,7 @@ void SBuildSelection::Construct(const FArguments& InArgs)
 		BuildUrl = CurrentBuildVersion.BuildUrl;
 		BuildId = CurrentBuildVersion.BuildId;
 		BuildVersion = CurrentBuildVersion.BuildVersion;
-		BranchName = SanitizeBucketSegment(CurrentBuildVersion.BranchName);
+		BranchName = SanitizeBranch(CurrentBuildVersion.BranchName);
 
 		EffectiveCompatibleChangelist = CurrentBuildVersion.GetEffectiveCompatibleChangelist();
 	}
@@ -143,7 +143,7 @@ void SBuildSelection::Construct(const FArguments& InArgs)
 	// regardless of past settings.
 	if (FParse::Value(FCommandLine::Get(), TEXT("-Branch"), BranchName))
 	{
-		BranchName = SanitizeBucketSegment(BranchName);
+		BranchName = SanitizeBranch(BranchName);
 	}
 
 	if (BranchName.IsEmpty())
@@ -175,7 +175,6 @@ void SBuildSelection::Construct(const FArguments& InArgs)
 		.HAlign(HAlign_Fill)
 		.VAlign(VAlign_Fill)
 		.Padding(0, 0, 0, 0)
-		.Expose(GridSlot)
 		[
 			GetGridPanel()
 		]
@@ -860,9 +859,28 @@ TSharedPtr<SWidget> SBuildSelection::OnGetBuildGroupContextMenuContent() const
 		NAME_None,
 		bRecursivelySearchable);
 
+
+	MenuBuilder.AddMenuEntry(
+		LOCTEXT("BuildSelection_CopyName", "Copy name"),
+		LOCTEXT("BuildSelection_CopyName_ToolTip", "Copies the name of the build group"),
+		FSlateIcon(),
+		FUIAction(
+			FExecuteAction::CreateLambda([this]
+				{
+					TArray<FBuildSelectionBuildGroupPtr> SelectedItems = BuildGroupListView->GetSelectedItems();
+					if (SelectedItems.Num() == 1)
+					{
+						FPlatformApplicationMisc::ClipboardCopy(*SelectedItems[0]->DisplayName);
+					}
+				})
+		),
+		NAME_None,
+		EUserInterfaceActionType::Button
+	);
+
 	MenuBuilder.AddMenuEntry(
 		LOCTEXT("BuildSelection_CopyCommit", "Copy commit"),
-		LOCTEXT("BuildSelection_CopyCommit_ToolTip", "Copies the commit or changelist number"),
+		LOCTEXT("BuildSelection_CopyCommit_ToolTip", "Copies the commit or changelist number of the build group"),
 		FSlateIcon(),
 		FUIAction(
 			FExecuteAction::CreateLambda([this]
@@ -1209,7 +1227,6 @@ void SBuildSelection::SetIncludePreflights(const bool bInIncludePreflights)
 
 bool SBuildSelection::WriteSetting(const FStringView InSectionName, const FStringView InKeyName, const FStringView InValue)
 {
-	static const FString StoreId = TEXT("Epic Games");
 	static const FString SectionName = TEXT("Unreal Engine/Build Storage");
 	const FString FinalSectionName = InSectionName.IsEmpty() ? SectionName : FPaths::Combine(SectionName, InSectionName);
 	const FString FinalKeyName(InKeyName);
@@ -1224,7 +1241,6 @@ bool SBuildSelection::WriteSetting(const FStringView InKeyName, const FStringVie
 
 bool SBuildSelection::ReadSetting(const FStringView InSectionName, const FStringView InKeyName, FString& OutValue)
 {
-	static const FString StoreId = TEXT("Epic Games");
 	static const FString SectionName = TEXT("Unreal Engine/Build Storage");
 	const FString FinalSectionName = InSectionName.IsEmpty() ? SectionName : FPaths::Combine(SectionName, InSectionName);
 	const FString FinalKeyName(InKeyName);
@@ -1296,7 +1312,14 @@ FString SBuildSelection::SanitizeForZenId(const FString& InString)
 	return OutString.IsEmpty() ? OutStringPreTrim : OutString;
 }
 
-FString SBuildSelection::SanitizeBucketSegment(const FString& InString)
+FString SBuildSelection::SanitizeBranch(const FString& InString)
+{
+	FString OutString(InString);
+	OutString.ReplaceCharInline(TCHAR('.'), TCHAR('-'));
+	return SanitizeBucket(OutString);
+}
+
+FString SBuildSelection::SanitizeBucket(const FString& InString)
 {
 	TStringBuilder<64> OutputBuilder;
 	for (int32 CharIndex = 0; CharIndex < InString.Len(); ++CharIndex)
