Build Storage Tool - Fixed issue were artifacts with short names (cooked output for platforms which name is shorter then 5 characters) was not displayed correctly.

#Synced-CL 47492933

diff --git a/Engine/Source/Editor/StorageServerWidgets/Private/SBuildSelection.cpp b/Engine/Source/Editor/StorageServerWidgets/Private/SBuildSelection.cpp
--- a/Engine/Source/Editor/StorageServerWidgets/Private/SBuildSelection.cpp
+++ b/Engine/Source/Editor/StorageServerWidgets/Private/SBuildSelection.cpp
@@ -566,11 +566,6 @@ void SBuildSelection::RegenerateBuildGroups(UE::Zen::Build::FListBuildsState& Li
 				{
 					ArtifactName.RightChopInline(TruncationIndex + 1);
 				}
-
-				if (ArtifactName.Len() < 5)
-				{
-					ArtifactName.Empty();
-				}
 			}
 			
 			if (BuildGroup->DisplayName.IsEmpty())
