BuildStorageTool: UX improvements.  Added configuration filtering, horizontal scroll bar for primary window, fixed a data race crash caused by calling RebuildLists from off the game thread while it may be in use on the game thread.

#Synced-CL 48523383

diff --git a/Engine/Source/Editor/StorageServerWidgets/Internal/SBuildSelection.h b/Engine/Source/Editor/StorageServerWidgets/Internal/SBuildSelection.h
--- a/Engine/Source/Editor/StorageServerWidgets/Internal/SBuildSelection.h
+++ b/Engine/Source/Editor/StorageServerWidgets/Internal/SBuildSelection.h
@@ -45,6 +45,12 @@ class SBuildSelection : public SCompoundWidget
 	SLATE_EVENT(FOnBuildTransferStarted, OnBuildTransferStarted)
 	SLATE_END_ARGS()
 
+	struct FArtifact
+	{
+		TSet<FName> Configurations;
+		UE::Zen::Build::FBuildServiceInstance::FBuildRecord BuildRecord;
+	};
+
 	struct FBuildGroup : public TSharedFromThis<FBuildGroup, ESPMode::ThreadSafe>
 	{
 		FString Namespace;
@@ -53,7 +59,7 @@ class SBuildSelection : public SCompoundWidget
 		FDateTime CreatedAt;
 		FString Category;
 		FString Job;
-		TMap<FString, UE::Zen::Build::FBuildServiceInstance::FBuildRecord> NamedArtifacts;
+		TMap<FString, FArtifact> NamedArtifacts;
 		bool bIsPreflight;
 	};
 
@@ -94,6 +100,8 @@ private:
 	void OnChooseTransferFilterClicked();
 	TSharedRef<SWidget> GetBuildDestinationPanel();
 	TSharedRef<SWidget> GetGridPanel();
+	void RebuildSelectedGroupArtifactHeader(TSharedPtr<FBuildGroup> Item);
+	void RebuildSelectedGroupArtifactGrid(TSharedPtr<FBuildGroup> Item);
 	void BuildGroupSelectionChanged(TSharedPtr<FBuildGroup> Item, ESelectInfo::Type SelectInfo);
 	EColumnSortMode::Type GetColumnSortMode(const FName ColumnId) const;
 	void OnColumnSortModeChanged(const EColumnSortPriority::Type SortPriority, const FName& ColumnId, const EColumnSortMode::Type InSortMode);
@@ -112,6 +120,7 @@ private:
 	void SetSelectedProject(const TSharedPtr<FString> InSelectedProject);
 	void SetSelectedBuildType(const TSharedPtr<FString> InBuildType);
 	void SetIncludePreflights(const bool bInIncludePreflights);
+	void SetSelectedConfiguration(const TSharedPtr<FString> InSelectedStream);
 
 	bool WriteSetting(FStringView InSectionName, FStringView InKeyName, FStringView InValue);
 	bool WriteSetting(FStringView InKeyName, FStringView InValue);
@@ -126,6 +135,8 @@ private:
 	static FString SanitizeBranch(const FString& InString);
 	static FString SanitizeBucket(const FString& InString);
 
+	static void ConformSelection(TSharedPtr<FString>& SelectedItem, const TArray<TSharedPtr<FString>>& SelectionList);
+
 	static const TCHAR* LexToString(EBuildType BuildType);
 
 	FReply ExploreDestination_OnClicked();
@@ -153,6 +164,7 @@ private:
 	TSharedPtr<SComboBox<TSharedPtr<FString>>> BuildTypeWidget;
 	TSharedPtr<SMultiSelectComboBox> PlatformFiltersWidget;
 
+	TSharedPtr<SGridPanel> SelectedGroupArtifactHeaderGrid;
 	TSharedPtr<SGridPanel> SelectedGroupArtifactGrid;
 	TArray<FString> SelectedGroupSelectedArtifacts;
 
@@ -161,7 +173,9 @@ private:
 	TArray<TSharedPtr<FString>> BuildTypeList;
 	TArray<TSharedPtr<FString>> PlatformList;
 
-	TSet<FString> ArtifactNameSet;
+	TMultiMap<FString, FString> ArtifactNameMap;
+	TSet<FString> ArtifactConfigurationSet;
+	TArray<TSharedPtr<FString>> ArtifactConfigurationList;
 
 	struct FTransferFilter
 	{
@@ -177,6 +191,7 @@ private:
 	TSharedPtr<FString> SelectedStream;
 	TSharedPtr<FString> SelectedProject;
 	TSharedPtr<FString> SelectedBuildType;
+	TSharedPtr<FString> SelectedConfiguration;
 	FString SelectedTextFilter;
 	bool bIncludePreflights = true;
 
diff --git a/Engine/Source/Editor/StorageServerWidgets/Private/SBuildSelection.cpp b/Engine/Source/Editor/StorageServerWidgets/Private/SBuildSelection.cpp
--- a/Engine/Source/Editor/StorageServerWidgets/Private/SBuildSelection.cpp
+++ b/Engine/Source/Editor/StorageServerWidgets/Private/SBuildSelection.cpp
@@ -164,7 +164,14 @@ void SBuildSelection::Construct(const FArguments& InArgs)
 
 	if (TSharedPtr<UE::Zen::Build::FBuildServiceInstance> ServiceInstance = BuildServiceInstance.Get())
 	{
-		ServiceInstance->OnRefreshNamespacesAndBucketsComplete().AddSP(this, &SBuildSelection::RebuildLists);
+		ServiceInstance->OnRefreshNamespacesAndBucketsComplete().AddSPLambda(this, [this]()
+			{
+				ExecuteOnGameThread(UE_SOURCE_LOCATION,
+					[this]
+					{
+						RebuildLists();
+					});
+			});
 	}
 
 	this->ChildSlot
@@ -249,22 +256,6 @@ void SBuildSelection::RebuildLists()
 				}
 			};
 
-		auto ConformSelection = [](TSharedPtr<FString>& SelectedItem, const TArray<TSharedPtr<FString>>& SelectionList)
-			{
-				if (!SelectedItem)
-				{
-					SelectedItem = SelectionList.IsEmpty() ? nullptr : SelectionList[0];
-					return;
-				}
-
-				const TSharedPtr<FString>* FoundSelectionListItem = SelectionList.FindByPredicate([&SelectedItem](const TSharedPtr<FString>& Item)
-					{
-						return *Item == *SelectedItem;
-					});
-
-				SelectedItem = FoundSelectionListItem ? *FoundSelectionListItem : SelectionList.IsEmpty() ? nullptr : SelectionList[0];
-			};
-
 		const uint32 SegmentIndexProject = 0;
 		const uint32 SegmentIndexBuildType = 1;
 		const uint32 SegmentIndexStream = 2;
@@ -494,7 +485,10 @@ void SBuildSelection::RegenerateBuildGroups(UE::Zen::Build::FListBuildsState& Li
 	using namespace UE::BuildSelection::Internal;
 
 	BuildGroups.Empty();
-	ArtifactNameSet.Empty();
+	ArtifactNameMap.Empty();
+	ArtifactConfigurationSet.Empty();
+	ArtifactConfigurationList.Empty();
+	bool bHasNoneArtifactConfiguration = false;
 
 	auto MakeGroupKey = [](const FString& Namespace, const FString& CommitIdentifier, const FBuildServiceInstance::FBuildRecord& BuildRecord)
 	{
@@ -567,6 +561,12 @@ void SBuildSelection::RegenerateBuildGroups(UE::Zen::Build::FListBuildsState& Li
 					ArtifactName.RightChopInline(TruncationIndex + 1);
 				}
 			}
+
+			FString Configuration;
+			if (FCbFieldView ConfigurationField = BuildRecord.Metadata.FindIgnoreCase("Configuration"); ConfigurationField.HasValue() && !ConfigurationField.HasError())
+			{
+				Configuration = FUTF8ToTCHAR(ConfigurationField.AsString());
+			}
 			
 			if (BuildGroup->DisplayName.IsEmpty())
 			{
@@ -603,12 +603,6 @@ void SBuildSelection::RegenerateBuildGroups(UE::Zen::Build::FListBuildsState& Li
 					}
 				}
 
-				FString Configuration;
-				if (FCbFieldView ConfigurationField = BuildRecord.Metadata.FindIgnoreCase("Configuration"); ConfigurationField.HasValue() && !ConfigurationField.HasError())
-				{
-					Configuration = FUTF8ToTCHAR(ConfigurationField.AsString());
-				}
-
 				FString ItemName;
 				FCbFieldView GroupNameView = BuildRecord.Metadata.FindIgnoreCase("buildgroup");
 				if (!GroupNameView.HasValue())
@@ -634,12 +628,42 @@ void SBuildSelection::RegenerateBuildGroups(UE::Zen::Build::FListBuildsState& Li
 				BuildGroup->bIsPreflight = bIsPreflight;
 			}
 
-			ArtifactNameSet.Add(ArtifactName);
-			BuildGroup->NamedArtifacts.FindOrAdd(ArtifactName, MoveTemp(BuildRecord));
+			TSet<FName> AllConfigurationNames;
+			if (Configuration.IsEmpty())
+			{
+				bHasNoneArtifactConfiguration = true;
+				ArtifactNameMap.EmplaceUnique(TEXT("None"), ArtifactName);
+			}
+			else
+			{
+				TArray<FString> AllConfigurations;
+				Configuration.ParseIntoArrayWS(AllConfigurations, TEXT("+"));
+				for (const FString& IndividualConfiguration : AllConfigurations)
+				{
+					AllConfigurationNames.Add(FName(IndividualConfiguration));
+					ArtifactConfigurationSet.Add(IndividualConfiguration);
+					ArtifactNameMap.EmplaceUnique(IndividualConfiguration, ArtifactName);
+				}
+			}
+			BuildGroup->NamedArtifacts.FindOrAdd(ArtifactName, { MoveTemp(AllConfigurationNames), MoveTemp(BuildRecord) });
 		}
 	}
 
-	ArtifactNameSet.StableSort(TLess());
+	ArtifactConfigurationSet.StableSort(TLess());
+	ArtifactNameMap.ValueStableSort(TLess());
+	if (!ArtifactConfigurationList.IsEmpty() && bHasNoneArtifactConfiguration)
+	{
+		ArtifactConfigurationList.Add(MakeShared<FString>(TEXT("None")));
+	}
+	if (ArtifactConfigurationSet.Num() > 1)
+	{
+		ArtifactConfigurationList.Add(MakeShared<FString>(TEXT("All")));
+	}
+	for (const FString& ArtifactConfiguration : ArtifactConfigurationSet)
+	{
+		ArtifactConfigurationList.Add(MakeShared<FString>(ArtifactConfiguration));
+	}
+	ConformSelection(SelectedConfiguration, ArtifactConfigurationList);
 	KeyedGroups.GenerateValueArray(BuildGroups);
 	ConditionalSortBuildGroups();
 }
@@ -764,11 +788,11 @@ bool SBuildSelection::DoesPassAllViewFilters(TSharedPtr<FBuildGroup> InItem) con
 	for (const FString& ActivePlatformFilter : ActivePlatformFilters)
 	{
 		bool bHasPlatform = false;
-		for (const TPair<FString, UE::Zen::Build::FBuildServiceInstance::FBuildRecord>& NamedArtifact : InItem->NamedArtifacts)
+		for (const TPair<FString, FArtifact>& NamedArtifact : InItem->NamedArtifacts)
 		{
-			if (NamedArtifact.Value.BucketId.EndsWith(ActivePlatformFilter) &&
-				NamedArtifact.Value.BucketId.Len() > ActivePlatformFilter.Len() &&
-				NamedArtifact.Value.BucketId[NamedArtifact.Value.BucketId.Len() - ActivePlatformFilter.Len() - 1] == TCHAR('.'))
+			if (NamedArtifact.Value.BuildRecord.BucketId.EndsWith(ActivePlatformFilter) &&
+				NamedArtifact.Value.BuildRecord.BucketId.Len() > ActivePlatformFilter.Len() &&
+				NamedArtifact.Value.BuildRecord.BucketId[NamedArtifact.Value.BuildRecord.BucketId.Len() - ActivePlatformFilter.Len() - 1] == TCHAR('.'))
 			{
 				bHasPlatform = true;
 				break;
@@ -897,10 +921,55 @@ TSharedPtr<SWidget> SBuildSelection::OnGetBuildGroupContextMenuContent() const
 	return MenuBuilder.MakeWidget();
 }
 
-void SBuildSelection::BuildGroupSelectionChanged(FBuildSelectionBuildGroupPtr Item, ESelectInfo::Type SelectInfo)
+void SBuildSelection::RebuildSelectedGroupArtifactHeader(FBuildSelectionBuildGroupPtr Item)
 {
-	using namespace UE::Zen::Build;
+	if (!SelectedGroupArtifactHeaderGrid)
+	{
+		return;
+	}
+
+	SelectedGroupArtifactHeaderGrid->ClearChildren();
+	if (ArtifactConfigurationSet.IsEmpty())
+	{
+		return;
+	}
 
+	SelectedGroupArtifactHeaderGrid->AddSlot(0, 0)
+		.Padding(10,0,2,0)
+		.VAlign(VAlign_Center)
+		[
+			SNew(STextBlock)
+			.Justification(ETextJustify::Right)
+			.Text(LOCTEXT("BuildSelection_Configuration", "Configuration:"))
+		];
+
+	SelectedGroupArtifactHeaderGrid->AddSlot(1, 0)
+		.Padding(0,0,0,0)
+		.VAlign(VAlign_Center)
+		[
+			SNew(SComboBox<TSharedPtr<FString>>)
+			.OptionsSource(&ArtifactConfigurationList)
+			.OnSelectionChanged_Lambda([this, Item](TSharedPtr<FString> NewSelectedConfiguration, ESelectInfo::Type SelectInfo)
+			{
+				if (NewSelectedConfiguration.IsValid() && *NewSelectedConfiguration != *SelectedConfiguration)
+				{
+					SetSelectedConfiguration(NewSelectedConfiguration);
+					RebuildSelectedGroupArtifactGrid(Item);
+				}
+			})
+			.OnGenerateWidget(this, &SBuildSelection::OnGenerateTextBlockFromString)
+			[
+				SNew(STextBlock)
+					.Text_Lambda([this]()
+					{
+						return FText::FromString(SelectedConfiguration ? **SelectedConfiguration : TEXT(""));
+					})
+			]
+		];
+}
+
+void SBuildSelection::RebuildSelectedGroupArtifactGrid(FBuildSelectionBuildGroupPtr Item)
+{
 	if (!SelectedGroupArtifactGrid)
 	{
 		return;
@@ -915,29 +984,88 @@ void SBuildSelection::BuildGroupSelectionChanged(FBuildSelectionBuildGroupPtr It
 
 	int32 Row = 0;
 	int32 Column = 0;
-	const int32 TotalListCount = ArtifactNameSet.Num();
+	TArray<FString> ArtifactNames;
+	const bool bPresentConsistentArtifactNameSet = false;
+	if (bPresentConsistentArtifactNameSet)
+	{
+		if (!SelectedConfiguration || ArtifactConfigurationSet.IsEmpty())
+		{
+			ArtifactNameMap.GenerateValueArray(ArtifactNames);
+		}
+		else if (*SelectedConfiguration == TEXT("All"))
+		{
+			for (const TPair<FString, FString>& ArtifactNamePair : ArtifactNameMap)
+			{
+				ArtifactNames.AddUnique(ArtifactNamePair.Value);
+			}
+			ArtifactNames.StableSort(TLess());
+		}
+		else
+		{
+			ArtifactNameMap.MultiFind(*SelectedConfiguration, ArtifactNames);
+		}
+	}
+	else
+	{
+		if (!SelectedConfiguration || ArtifactConfigurationSet.IsEmpty() || *SelectedConfiguration == TEXT("All"))
+		{
+			Item->NamedArtifacts.GenerateKeyArray(ArtifactNames);
+		}
+		else
+		{
+			FName SelectedConfigurationName(*SelectedConfiguration);
+			for (const TPair<FString, FArtifact>& NameArtifactPair : Item->NamedArtifacts)
+			{
+				if (NameArtifactPair.Value.Configurations.Contains(SelectedConfigurationName))
+				{
+					ArtifactNames.Add(NameArtifactPair.Key);
+				}
+			}
+		}
+		ArtifactNames.StableSort(TLess());
+	}
+
+	const int32 TotalListCount = ArtifactNames.Num();
 	// Heuristic for choosing grid rows when laying out rows and columns for the quantity of platforms
 	const int32 MaxRows = FMath::Max(4, FMath::FloorToInt32(FMath::Pow(float(TotalListCount), 0.7)));
 
-	for (const FString& ArtifactName : ArtifactNameSet)
+	auto IsBuildGroupArtifactEnabled = [this, Item] (const FString& ArtifactName)
+	{
+		const FArtifact* Artifact = Item->NamedArtifacts.Find(*ArtifactName);
+		if (!Artifact)
+		{
+			return false;
+		}
+		if (ArtifactConfigurationSet.IsEmpty() || !SelectedConfiguration || *SelectedConfiguration == TEXT("All"))
+		{
+			return true;
+		}
+		if (*SelectedConfiguration == TEXT("None"))
+		{
+			return Artifact->Configurations.IsEmpty();
+		}
+		return Artifact->Configurations.Contains(FName(*SelectedConfiguration));
+	};
+
+	for (const FString& ArtifactName : ArtifactNames)
 	{
 		TSharedPtr<SCheckBox> CurrentCheckbox;
 		SelectedGroupArtifactGrid->AddSlot(Column, Row)
 			.Padding(0,0,0,2)
 			[
 				SNew(SHorizontalBox)
-				.IsEnabled_Lambda([this, Item, ArtifactName]
+				.IsEnabled_Lambda([this, Item, ArtifactName, IsBuildGroupArtifactEnabled]
 				{
-					return Item->NamedArtifacts.Contains(*ArtifactName);
+					return IsBuildGroupArtifactEnabled(ArtifactName);
 				})
 				+SHorizontalBox::Slot()
 				.AutoWidth()
 				.VAlign(VAlign_Center)
 				[
 					SAssignNew(CurrentCheckbox, SCheckBox)
-					.IsChecked_Lambda([this, Item, ArtifactName]
+					.IsChecked_Lambda([this, Item, ArtifactName, IsBuildGroupArtifactEnabled]
 					{
-						bool bIsChecked = Item->NamedArtifacts.Contains(ArtifactName) && SelectedGroupSelectedArtifacts.Contains(ArtifactName);
+						bool bIsChecked = IsBuildGroupArtifactEnabled(ArtifactName) && SelectedGroupSelectedArtifacts.Contains(ArtifactName);
 						return bIsChecked ? ECheckBoxState::Checked : ECheckBoxState::Unchecked;
 					})
 					.OnCheckStateChanged_Lambda([this, ArtifactName](ECheckBoxState InNewState)
@@ -991,22 +1119,30 @@ void SBuildSelection::BuildGroupSelectionChanged(FBuildSelectionBuildGroupPtr It
 					SNew(SButton)
 					.ButtonStyle(&FAppStyle::Get().GetWidgetStyle<FButtonStyle>("SimpleButton"))
 					.ToolTipText(LOCTEXT("BuildSelection_BrowseToBuildUrl", "Browse to build url"))
-					.IsEnabled_Lambda([this, Item, ArtifactName]
+					.IsEnabled_Lambda([this, Item, ArtifactName, IsBuildGroupArtifactEnabled]
 					{
-						if (FBuildServiceInstance::FBuildRecord* BuildRecord = Item->NamedArtifacts.Find(ArtifactName))
+						if (!IsBuildGroupArtifactEnabled(ArtifactName))
+						{
+							return false;
+						}
+						if (FArtifact* Artifact = Item->NamedArtifacts.Find(ArtifactName))
 						{
-							if (FCbFieldView BuildUrlField = BuildRecord->Metadata["buildurl"]; BuildUrlField.HasValue() && !BuildUrlField.HasError())
+							if (FCbFieldView BuildUrlField = Artifact->BuildRecord.Metadata["buildurl"]; BuildUrlField.HasValue() && !BuildUrlField.HasError())
 							{
 								return true;
 							}
 						}
 						return false;
 					})
-					.OnClicked_Lambda([this, Item, ArtifactName]
+					.OnClicked_Lambda([this, Item, ArtifactName, IsBuildGroupArtifactEnabled]
 					{
-						if (FBuildServiceInstance::FBuildRecord* BuildRecord = Item->NamedArtifacts.Find(ArtifactName))
+						if (!IsBuildGroupArtifactEnabled(ArtifactName))
 						{
-							if (FCbFieldView BuildUrlField = BuildRecord->Metadata["buildurl"]; BuildUrlField.HasValue() && !BuildUrlField.HasError())
+							return FReply::Unhandled();
+						}
+						if (FArtifact* Artifact = Item->NamedArtifacts.Find(ArtifactName))
+						{
+							if (FCbFieldView BuildUrlField = Artifact->BuildRecord.Metadata["buildurl"]; BuildUrlField.HasValue() && !BuildUrlField.HasError())
 							{
 								FString Url(FUTF8ToTCHAR(BuildUrlField.AsString()));
 								FPlatformProcess::LaunchURL(*Url, nullptr, nullptr);
@@ -1031,6 +1167,15 @@ void SBuildSelection::BuildGroupSelectionChanged(FBuildSelectionBuildGroupPtr It
 	}
 }
 
+void SBuildSelection::BuildGroupSelectionChanged(FBuildSelectionBuildGroupPtr Item, ESelectInfo::Type SelectInfo)
+{
+	using namespace UE::Zen::Build;
+
+	RebuildSelectedGroupArtifactHeader(Item);
+
+	RebuildSelectedGroupArtifactGrid(Item);
+}
+
 EColumnSortMode::Type SBuildSelection::GetColumnSortMode(const FName ColumnId) const
 {
 	if (BuildGroupSortByColumn == ColumnId)
@@ -1220,6 +1365,12 @@ void SBuildSelection::SetIncludePreflights(const bool bInIncludePreflights)
 	PostViewFilterChange();
 }
 
+void SBuildSelection::SetSelectedConfiguration(const TSharedPtr<FString> InSelectedConfiguration)
+{
+	SelectedConfiguration = InSelectedConfiguration;
+	WriteSetting(TEXT("Configuration"), *SelectedConfiguration);
+}
+
 bool SBuildSelection::WriteSetting(const FStringView InSectionName, const FStringView InKeyName, const FStringView InValue)
 {
 	static const FString SectionName = TEXT("Unreal Engine/Build Storage");
@@ -1355,6 +1506,22 @@ FString SBuildSelection::SanitizeBucket(const FString& InString)
 	return OutString;
 }
 
+void SBuildSelection::ConformSelection(TSharedPtr<FString>& SelectedItem, const TArray<TSharedPtr<FString>>& SelectionList)
+{
+	if (!SelectedItem)
+	{
+		SelectedItem = SelectionList.IsEmpty() ? nullptr : SelectionList[0];
+		return;
+	}
+
+	const TSharedPtr<FString>* FoundSelectionListItem = SelectionList.FindByPredicate([&SelectedItem](const TSharedPtr<FString>& Item)
+		{
+			return *Item == *SelectedItem;
+		});
+
+	SelectedItem = FoundSelectionListItem ? *FoundSelectionListItem : SelectionList.IsEmpty() ? nullptr : SelectionList[0];
+}
+
 const TCHAR* SBuildSelection::LexToString(EBuildType BuildType)
 {
 	static const TCHAR* Strings[] =
@@ -1406,18 +1573,28 @@ void SBuildSelection::DownloadSelectedArtifactsInSelectedBuildGroup()
 			return;
 		}
 
+		FName SelectedConfigurationName(NAME_None);
+		if (SelectedConfiguration && *SelectedConfiguration != TEXT("All"))
+		{
+			SelectedConfigurationName = FName(*SelectedConfiguration);
+		}
 		for (const FString& ArtifactForSelectedGroup : SelectedGroupSelectedArtifacts)
 		{
-			if (FBuildServiceInstance::FBuildRecord* BuildRecord = SelectedItems[0]->NamedArtifacts.Find(ArtifactForSelectedGroup))
+			if (FArtifact* Artifact = SelectedItems[0]->NamedArtifacts.Find(ArtifactForSelectedGroup))
 			{
+				if ((SelectedConfigurationName != NAME_None) && !Artifact->Configurations.Contains(SelectedConfigurationName))
+				{
+					continue;
+				}
+
 				FString DestinationPlatformName = ArtifactForSelectedGroup;
-				if (FCbFieldView CookPlatformField = BuildRecord->Metadata["cookPlatform"]; CookPlatformField.HasValue() && !CookPlatformField.HasError())
+				if (FCbFieldView CookPlatformField = Artifact->BuildRecord.Metadata["cookPlatform"]; CookPlatformField.HasValue() && !CookPlatformField.HasError())
 				{
 					DestinationPlatformName = *WriteToString<64>(CookPlatformField.AsString());
 				}
 
 				FName TargetPlatformName = NAME_None;
-				if (FCbFieldView PlatformField = BuildRecord->Metadata["platform"]; PlatformField.HasValue() && !PlatformField.HasError())
+				if (FCbFieldView PlatformField = Artifact->BuildRecord.Metadata["platform"]; PlatformField.HasValue() && !PlatformField.HasError())
 				{
 					TargetPlatformName = FName(PlatformField.AsString());
 				}
@@ -1454,7 +1631,7 @@ void SBuildSelection::DownloadSelectedArtifactsInSelectedBuildGroup()
 						DestinationOplogId = FString::Printf(TEXT("%s.%s"), *DestinationOplogId, *SanitizeForZenId(SelectedItems[0]->DisplayName));
 					}
 					FBuildServiceInstance::FBuildTransfer BuildTransfer =
-						ServiceInstance->StartOplogBuildTransfer(BuildRecord->BuildId, TransferName, DestinationProjectId, DestinationOplogId, RootDir, EngineDir, ProjectFilename, SelectedItems[0]->Namespace, BuildRecord->BucketId, TargetPlatformName);
+						ServiceInstance->StartOplogBuildTransfer(Artifact->BuildRecord.BuildId, TransferName, DestinationProjectId, DestinationOplogId, RootDir, EngineDir, ProjectFilename, SelectedItems[0]->Namespace, Artifact->BuildRecord.BucketId, TargetPlatformName);
 					OnBuildTransferStarted.ExecuteIfBound(BuildTransfer, SelectedItems[0]->DisplayName, ArtifactForSelectedGroup);
 				}
 				else
@@ -1470,11 +1647,11 @@ void SBuildSelection::DownloadSelectedArtifactsInSelectedBuildGroup()
 					}
 					FBuildServiceInstance::FBuildTransfer BuildTransfer =
 						ServiceInstance->StartBuildTransfer(
-							BuildRecord->BuildId,
+							Artifact->BuildRecord.BuildId,
 							TransferName,
 							DestinationFolder,
 							SelectedItems[0]->Namespace,
-							BuildRecord->BucketId,
+							Artifact->BuildRecord.BucketId,
 							UserSelectedTransferFilters[BuildTypeIndex].IncludeFilter,
 							UserSelectedTransferFilters[BuildTypeIndex].ExcludeFilter);
 					OnBuildTransferStarted.ExecuteIfBound(BuildTransfer, SelectedItems[0]->DisplayName, ArtifactForSelectedGroup);
@@ -2156,10 +2333,24 @@ TSharedRef<SWidget> SBuildSelection::GetGridPanel()
 			.VAlign(VAlign_Top)
 			.AutoHeight()
 			[
-				SNew(STextBlock)
-					.ColorAndOpacity(TitleColor)
-					.Font(TitleFont)
-					.Text(LOCTEXT("BuildSelection_AvailableBuilds", "Available Builds"))
+				SNew(SHorizontalBox)
+				+SHorizontalBox::Slot()
+				.HAlign(HAlign_Left)
+				.VAlign(VAlign_Center)
+				.AutoWidth()
+				[
+					SNew(STextBlock)
+						.ColorAndOpacity(TitleColor)
+						.Font(TitleFont)
+						.Text(LOCTEXT("BuildSelection_AvailableBuilds", "Available Builds"))
+				]
+				+SHorizontalBox::Slot()
+				.HAlign(HAlign_Fill)
+				.VAlign(VAlign_Center)
+				.Padding(0, 0, 0, 0)
+				[
+					SAssignNew(SelectedGroupArtifactHeaderGrid, SGridPanel)
+				]
 			]
 			+SVerticalBox::Slot()
 			.HAlign(HAlign_Fill)
@@ -2237,11 +2428,19 @@ TSharedRef<SWidget> SBuildSelection::GetGridPanel()
 						return false;
 					}
 
+					FName SelectedConfigurationName(NAME_None);
+					if (SelectedConfiguration && *SelectedConfiguration != TEXT("All"))
+					{
+						SelectedConfigurationName = FName(*SelectedConfiguration);
+					}
 					for (const FString& ArtifactForSelectedGroup : SelectedGroupSelectedArtifacts)
 					{
-						if (SelectedItems[0]->NamedArtifacts.Contains(ArtifactForSelectedGroup))
+						if (FArtifact* FoundArtifact = SelectedItems[0]->NamedArtifacts.Find(ArtifactForSelectedGroup))
 						{
-							return true;
+							if ((SelectedConfigurationName == NAME_None) || FoundArtifact->Configurations.Contains(SelectedConfigurationName))
+							{
+								return true;
+							}
 						}
 					}
 					return false;
diff --git a/Engine/Source/Programs/BuildStorageTool/Private/BuildStorageTool.cpp b/Engine/Source/Programs/BuildStorageTool/Private/BuildStorageTool.cpp
--- a/Engine/Source/Programs/BuildStorageTool/Private/BuildStorageTool.cpp
+++ b/Engine/Source/Programs/BuildStorageTool/Private/BuildStorageTool.cpp
@@ -28,6 +28,7 @@
 #include "Widgets/Input/SButton.h"
 #include "Widgets/Input/SButton.h"
 #include "Widgets/Input/SCheckBox.h"
+#include "Widgets/Layout/SScrollBox.h"
 #include "Widgets/Notifications/SNotificationList.h"
 #include "Widgets/Text/STextBlock.h"
 #include "ZenServiceInstanceManager.h"
@@ -307,7 +308,7 @@ public:
 		Window =
 			SNew(SWindow)
 			.Title(GetWindowTitle())
-			.ClientSize(FVector2D(1000.0f, 600.0f))
+			.ClientSize(FVector2D(1280.0f, 720.0f))
 			.ActivationPolicy(EWindowActivationPolicy::Always)
 			.SizingRule(ESizingRule::UserSized)
 			.IsTopmostWindow(false)
@@ -319,53 +320,59 @@ public:
 				SNew(SBorder)
 				.BorderImage(FCoreStyle::Get().GetBrush("ToolPanel.GroupBorder"))
 				[
-					SNew(SVerticalBox)
-
-					// Menu
-					+ SVerticalBox::Slot()
-					.AutoHeight()
-					[
-						MakeMainMenu()
-					]
-
-					// Login panel
-					+ SVerticalBox::Slot()
-					.AutoHeight()
-					.Padding(0.0f, 10.0f, 5.0f, 0.0f)
-					.HAlign(HAlign_Fill)
-					.VAlign(VAlign_Top)
-					[
-						SNew(SBuildLogin)
-						.ZenServiceInstance(ZenServiceInstanceManager.ToSharedRef(), &UE::Zen::FServiceInstanceManager::GetZenServiceInstance)
-						.BuildServiceInstance(BuildServiceInstanceManager.ToSharedRef(), &UE::Zen::Build::FServiceInstanceManager::GetBuildServiceInstance)
-					]
-
-					// Selection panel
-					+ SVerticalBox::Slot()
-					.Padding(0.0f, 10.0f, 5.0f, 10.0f)
-					.HAlign(HAlign_Fill)
-					.VAlign(VAlign_Fill)
-					[
-						SNew(SBuildSelection)
-						.ZenServiceInstance(ZenServiceInstanceManager.ToSharedRef(), &UE::Zen::FServiceInstanceManager::GetZenServiceInstance)
-						.BuildServiceInstance(BuildServiceInstanceManager.ToSharedRef(), &UE::Zen::Build::FServiceInstanceManager::GetBuildServiceInstance)
-						.OnBuildTransferStarted_Lambda([this]
-							(UE::Zen::Build::FBuildServiceInstance::FBuildTransfer Transfer, FStringView Name, FStringView Platform)
-							{
-								BuildActivity->AddBuildTransfer(Transfer, Name, Platform);
-							})
-					]
-
-					// Activity panel
-					+ SVerticalBox::Slot()
-					.AutoHeight()
-					.Padding(0.0f, 10.0f, 5.0f, 0.0f)
-					.HAlign(HAlign_Fill)
-					.VAlign(VAlign_Bottom)
+					SNew(SScrollBox)
+					.Orientation(EOrientation::Orient_Horizontal)
+					+ SScrollBox::Slot()
+					.FillSize(1.0f)
 					[
-						SAssignNew(BuildActivity, SBuildActivity)
-						.ZenServiceInstance(ZenServiceInstanceManager.ToSharedRef(), &UE::Zen::FServiceInstanceManager::GetZenServiceInstance)
-						.BuildServiceInstance(BuildServiceInstanceManager.ToSharedRef(), &UE::Zen::Build::FServiceInstanceManager::GetBuildServiceInstance)
+						SNew(SVerticalBox)
+
+						// Menu
+						+ SVerticalBox::Slot()
+						.AutoHeight()
+						[
+							MakeMainMenu()
+						]
+
+						// Login panel
+						+ SVerticalBox::Slot()
+						.AutoHeight()
+						.Padding(0.0f, 10.0f, 5.0f, 0.0f)
+						.HAlign(HAlign_Fill)
+						.VAlign(VAlign_Top)
+						[
+							SNew(SBuildLogin)
+							.ZenServiceInstance(ZenServiceInstanceManager.ToSharedRef(), &UE::Zen::FServiceInstanceManager::GetZenServiceInstance)
+							.BuildServiceInstance(BuildServiceInstanceManager.ToSharedRef(), &UE::Zen::Build::FServiceInstanceManager::GetBuildServiceInstance)
+						]
+
+						// Selection panel
+						+ SVerticalBox::Slot()
+						.Padding(0.0f, 10.0f, 5.0f, 10.0f)
+						.HAlign(HAlign_Fill)
+						.VAlign(VAlign_Fill)
+						[
+							SNew(SBuildSelection)
+							.ZenServiceInstance(ZenServiceInstanceManager.ToSharedRef(), &UE::Zen::FServiceInstanceManager::GetZenServiceInstance)
+							.BuildServiceInstance(BuildServiceInstanceManager.ToSharedRef(), &UE::Zen::Build::FServiceInstanceManager::GetBuildServiceInstance)
+							.OnBuildTransferStarted_Lambda([this]
+								(UE::Zen::Build::FBuildServiceInstance::FBuildTransfer Transfer, FStringView Name, FStringView Platform)
+								{
+									BuildActivity->AddBuildTransfer(Transfer, Name, Platform);
+								})
+						]
+
+						// Activity panel
+						+ SVerticalBox::Slot()
+						.AutoHeight()
+						.Padding(0.0f, 10.0f, 5.0f, 0.0f)
+						.HAlign(HAlign_Fill)
+						.VAlign(VAlign_Bottom)
+						[
+							SAssignNew(BuildActivity, SBuildActivity)
+							.ZenServiceInstance(ZenServiceInstanceManager.ToSharedRef(), &UE::Zen::FServiceInstanceManager::GetZenServiceInstance)
+							.BuildServiceInstance(BuildServiceInstanceManager.ToSharedRef(), &UE::Zen::Build::FServiceInstanceManager::GetBuildServiceInstance)
+						]
 					]
 				]
 			];
