BuildStorageTool: fixes for the way Configuration is conformed and persisted.

#Synced-CL 48523383

diff --git a/Engine/Source/Editor/StorageServerWidgets/Private/SBuildSelection.cpp b/Engine/Source/Editor/StorageServerWidgets/Private/SBuildSelection.cpp
--- a/Engine/Source/Editor/StorageServerWidgets/Private/SBuildSelection.cpp
+++ b/Engine/Source/Editor/StorageServerWidgets/Private/SBuildSelection.cpp
@@ -157,6 +157,11 @@ void SBuildSelection::Construct(const FArguments& InArgs)
 		bIncludePreflights = FCString::ToBool(*IncludePreflightsString);
 	}
 
+	if (FString PastConfiguration; ReadSetting(TEXT("Configuration"), PastConfiguration))
+	{
+		SelectedConfiguration = MakeShared<FString>(MoveTemp(PastConfiguration));
+	}
+
 	for (bool& bAppendBuildNameToDestination : bAppendBuildNameToDestinations)
 	{
 		bAppendBuildNameToDestination = false;
@@ -663,7 +668,11 @@ void SBuildSelection::RegenerateBuildGroups(UE::Zen::Build::FListBuildsState& Li
 	{
 		ArtifactConfigurationList.Add(MakeShared<FString>(ArtifactConfiguration));
 	}
-	ConformSelection(SelectedConfiguration, ArtifactConfigurationList);
+	if (!ArtifactConfigurationList.IsEmpty())
+	{
+		ConformSelection(SelectedConfiguration, ArtifactConfigurationList);
+		WriteSetting(TEXT("Configuration"), SelectedConfiguration ? **SelectedConfiguration : TEXT(""));
+	}
 	KeyedGroups.GenerateValueArray(BuildGroups);
 	ConditionalSortBuildGroups();
 }
