Make the "Append build name" setting persisted on a per-BuildType basis.  If you set it for a BuildType, the next run will have it enabled for that BuildType.

#Synced-CL 48523383

diff --git a/Engine/Source/Editor/StorageServerWidgets/Internal/SBuildSelection.h b/Engine/Source/Editor/StorageServerWidgets/Internal/SBuildSelection.h
--- a/Engine/Source/Editor/StorageServerWidgets/Internal/SBuildSelection.h
+++ b/Engine/Source/Editor/StorageServerWidgets/Internal/SBuildSelection.h
@@ -114,6 +114,7 @@ private:
 	FString GetDefaultDestination() const;
 	FString GetEffectiveDestination() const;
 
+	void SetUserSelectedAppendBuildNameToDestination(bool bAppendBuildNameToDestination);
 	void SetUserSelectedTransferFilter(const FStringView InIncludeFilter, const FStringView InExcludeFilter);
 
 	void SetSelectedStream(const TSharedPtr<FString> InSelectedStream);
diff --git a/Engine/Source/Editor/StorageServerWidgets/Private/SBuildSelection.cpp b/Engine/Source/Editor/StorageServerWidgets/Private/SBuildSelection.cpp
--- a/Engine/Source/Editor/StorageServerWidgets/Private/SBuildSelection.cpp
+++ b/Engine/Source/Editor/StorageServerWidgets/Private/SBuildSelection.cpp
@@ -108,9 +108,17 @@ void SBuildSelection::Construct(const FArguments& InArgs)
 		}
 	}
 
+	for (bool& bAppendBuildNameToDestination : bAppendBuildNameToDestinations)
+	{
+		bAppendBuildNameToDestination = false;
+	}
+
 	for (int32 BuildTypeIndex = 0; BuildTypeIndex < (int32)EBuildType::Count; ++BuildTypeIndex)
 	{
 		ReadSetting(LexToString((EBuildType)BuildTypeIndex), TEXT("UserSelectedDestination"), UserSelectedDestinations[BuildTypeIndex]);
+		FString AppendBuildNameToDestinationString;
+		ReadSetting(LexToString((EBuildType)BuildTypeIndex), TEXT("UserSelectedAppendBuildNameToDestination"), AppendBuildNameToDestinationString);
+		bAppendBuildNameToDestinations[BuildTypeIndex] = FCString::ToBool(*AppendBuildNameToDestinationString);
 		ReadSetting(LexToString((EBuildType)BuildTypeIndex), TEXT("UserSelectedIncludeFilter"), UserSelectedTransferFilters[BuildTypeIndex].IncludeFilter);
 		ReadSetting(LexToString((EBuildType)BuildTypeIndex), TEXT("UserSelectedExcludeFilter"), UserSelectedTransferFilters[BuildTypeIndex].ExcludeFilter);
 	}
@@ -162,11 +170,6 @@ void SBuildSelection::Construct(const FArguments& InArgs)
 		SelectedConfiguration = MakeShared<FString>(MoveTemp(PastConfiguration));
 	}
 
-	for (bool& bAppendBuildNameToDestination : bAppendBuildNameToDestinations)
-	{
-		bAppendBuildNameToDestination = false;
-	}
-
 	if (TSharedPtr<UE::Zen::Build::FBuildServiceInstance> ServiceInstance = BuildServiceInstance.Get())
 	{
 		ServiceInstance->OnRefreshNamespacesAndBucketsComplete().AddSPLambda(this, [this]()
@@ -1263,6 +1266,15 @@ void SBuildSelection::SetUserSelectedDestination(const FStringView InDestination
 	WriteSetting(LexToString(BuildType), TEXT("UserSelectedDestination"), InDestination);
 }
 
+void SBuildSelection::SetUserSelectedAppendBuildNameToDestination(bool bAppendBuildNameToDestination)
+{
+	EBuildType BuildType = GetSelectedBuildType();
+	int BuildTypeIndex = (int)BuildType;
+	bAppendBuildNameToDestinations[BuildTypeIndex] = bAppendBuildNameToDestination;
+
+	WriteSetting(LexToString(BuildType), TEXT("UserSelectedAppendBuildNameToDestination"), ::LexToString(bAppendBuildNameToDestination));
+}
+
 void SBuildSelection::SetUserSelectedTransferFilter(const FStringView InIncludeFilter, const FStringView InExcludeFilter)
 {
 	EBuildType BuildType = GetSelectedBuildType();
@@ -1888,11 +1900,11 @@ TSharedRef<SWidget> SBuildSelection::GetBuildDestinationPanel()
 				int BuildTypeIndex = (int)BuildType;
 				if (InNewState == ECheckBoxState::Checked)
 				{
-					bAppendBuildNameToDestinations[BuildTypeIndex] = true;
+					SetUserSelectedAppendBuildNameToDestination(true);
 				}
 				else if (InNewState == ECheckBoxState::Unchecked)
 				{
-					bAppendBuildNameToDestinations[BuildTypeIndex] = false;
+					SetUserSelectedAppendBuildNameToDestination(false);
 				}
 			})
 		]
@@ -1915,11 +1927,11 @@ TSharedRef<SWidget> SBuildSelection::GetBuildDestinationPanel()
 
 					if (NewState == ECheckBoxState::Checked)
 					{
-						bAppendBuildNameToDestinations[BuildTypeIndex] = true;
+						SetUserSelectedAppendBuildNameToDestination(true);
 					}
 					else if (NewState == ECheckBoxState::Unchecked)
 					{
-						bAppendBuildNameToDestinations[BuildTypeIndex] = false;
+						SetUserSelectedAppendBuildNameToDestination(false);
 					}
 					return FReply::Handled();
 				})
