BuildStorageTool: Limit the Destination text box to a maximum size of 500 to avoid long destination path text from blowing out the width of the overall build selection panel.  Also moved the ScrollBox container from the overall BuildStorageTool window down into the bottom section of the build selection panel.  This way if the build selection panel can't shrink to the size required by its containing window, the panel alone gets made horizontally scrollable, not the whole window.  This avoids the other panels having their flow influenced by a long list of available builds or a long destination path.

#Synced-CL 48523383

diff --git a/Engine/Source/Editor/StorageServerWidgets/Private/SBuildSelection.cpp b/Engine/Source/Editor/StorageServerWidgets/Private/SBuildSelection.cpp
--- a/Engine/Source/Editor/StorageServerWidgets/Private/SBuildSelection.cpp
+++ b/Engine/Source/Editor/StorageServerWidgets/Private/SBuildSelection.cpp
@@ -26,7 +26,9 @@
 #include "Widgets/Input/SCheckBox.h"
 #include "Widgets/Input/SEditableTextBox.h"
 #include "Widgets/Input/SHyperlink.h"
+#include "Widgets/Layout/SBox.h"
 #include "Widgets/Layout/SGridPanel.h"
+#include "Widgets/Layout/SScrollBox.h"
 #include "Widgets/SBoxPanel.h"
 #include "Widgets/Text/STextBlock.h"
 #include "Widgets/Views/SHeaderRow.h"
@@ -1758,21 +1760,25 @@ TSharedRef<SWidget> SBuildSelection::GetBuildDestinationPanel()
 			.AutoHeight()
 			.VAlign(VAlign_Top)
 			[
-				SNew(SEditableTextBox)
-				.OverflowPolicy(ETextOverflowPolicy::MiddleEllipsis)
-				.MinDesiredWidth(200.0f)
-				.HintText_Lambda([this]()
-				{
-					return FText::FromString(GetDefaultDestination());
-				})
-				.Text_Lambda([this]()
-				{
-					return FText::FromString(GetUserSelectedDestination());
-				})
-				.OnTextCommitted_Lambda([this](const FText& Text, const ETextCommit::Type CommitType)
-				{
-					SetUserSelectedDestination(Text.ToString());
-				})
+				SNew(SBox)
+				.MaxDesiredWidth(500.0f)
+				[
+					SNew(SEditableTextBox)
+					.OverflowPolicy(ETextOverflowPolicy::MiddleEllipsis)
+					.MinDesiredWidth(200.0f)
+					.HintText_Lambda([this]()
+					{
+						return FText::FromString(GetDefaultDestination());
+					})
+					.Text_Lambda([this]()
+					{
+						return FText::FromString(GetUserSelectedDestination());
+					})
+					.OnTextCommitted_Lambda([this](const FText& Text, const ETextCommit::Type CommitType)
+					{
+						SetUserSelectedDestination(Text.ToString());
+					})
+				]
 			]
 
 			+ SVerticalBox::Slot()
@@ -2339,138 +2345,144 @@ TSharedRef<SWidget> SBuildSelection::GetGridPanel()
 	.VAlign(VAlign_Bottom)
 	[
 
-		SNew(SHorizontalBox)
-		.Visibility_Lambda([this]()
-		{
-			return !!BuildListRefreshesInProgress || BuildGroupListView->GetNumItemsSelected() == 0 ? EVisibility::Collapsed : EVisibility::Visible;
-		})
-		+SHorizontalBox::Slot()
-		.HAlign(HAlign_Left)
-		.AutoWidth()
+		SNew(SScrollBox)
+		.Orientation(EOrientation::Orient_Horizontal)
+		+ SScrollBox::Slot()
+		.FillSize(1.0f)
 		[
-			SNew(SVerticalBox)
-			+SVerticalBox::Slot()
-			.HAlign(HAlign_Fill)
-			.VAlign(VAlign_Top)
-			.AutoHeight()
+			SNew(SHorizontalBox)
+			.Visibility_Lambda([this]()
+			{
+				return !!BuildListRefreshesInProgress || BuildGroupListView->GetNumItemsSelected() == 0 ? EVisibility::Collapsed : EVisibility::Visible;
+			})
+			+SHorizontalBox::Slot()
+			.HAlign(HAlign_Left)
+			.AutoWidth()
 			[
-				SNew(SHorizontalBox)
-				+SHorizontalBox::Slot()
-				.HAlign(HAlign_Left)
-				.VAlign(VAlign_Center)
-				.AutoWidth()
+				SNew(SVerticalBox)
+				+SVerticalBox::Slot()
+				.HAlign(HAlign_Fill)
+				.VAlign(VAlign_Top)
+				.AutoHeight()
 				[
-					SNew(STextBlock)
-						.ColorAndOpacity(TitleColor)
-						.Font(TitleFont)
-						.Text(LOCTEXT("BuildSelection_AvailableBuilds", "Available Builds"))
+					SNew(SHorizontalBox)
+					+SHorizontalBox::Slot()
+					.HAlign(HAlign_Left)
+					.VAlign(VAlign_Center)
+					.AutoWidth()
+					[
+						SNew(STextBlock)
+							.ColorAndOpacity(TitleColor)
+							.Font(TitleFont)
+							.Text(LOCTEXT("BuildSelection_AvailableBuilds", "Available Builds"))
+					]
+					+SHorizontalBox::Slot()
+					.HAlign(HAlign_Fill)
+					.VAlign(VAlign_Center)
+					.Padding(0, 0, 0, 0)
+					[
+						SAssignNew(SelectedGroupArtifactHeaderGrid, SGridPanel)
+					]
 				]
-				+SHorizontalBox::Slot()
+				+SVerticalBox::Slot()
 				.HAlign(HAlign_Fill)
-				.VAlign(VAlign_Center)
-				.Padding(0, 0, 0, 0)
+				.VAlign(VAlign_Top)
+				.AutoHeight()
 				[
-					SAssignNew(SelectedGroupArtifactHeaderGrid, SGridPanel)
+					SAssignNew(SelectedGroupArtifactGrid, SGridPanel)
 				]
 			]
-			+SVerticalBox::Slot()
-			.HAlign(HAlign_Fill)
-			.VAlign(VAlign_Top)
-			.AutoHeight()
-			[
-				SAssignNew(SelectedGroupArtifactGrid, SGridPanel)
-			]
-		]
-		+SHorizontalBox::Slot()
-		.HAlign(HAlign_Right)
-		.FillWidth(0.5f)
-		[
-			SNew(SVerticalBox)
-			+SVerticalBox::Slot()
-			.HAlign(HAlign_Fill)
-			.VAlign(VAlign_Top)
-			.AutoHeight()
+			+SHorizontalBox::Slot()
+			.HAlign(HAlign_Right)
+			.FillWidth(0.5f)
 			[
-				SNew(SHorizontalBox)
-				+SHorizontalBox::Slot()
-				.HAlign(HAlign_Left)
-				.VAlign(VAlign_Bottom)
-				.AutoWidth()
-				[
-					SNew(STextBlock)
-						.ColorAndOpacity(TitleColor)
-						.Font(TitleFont)
-						.Text(LOCTEXT("BuildSelection_Destination", "Destination"))
-				]
-				+SHorizontalBox::Slot()
-				.HAlign(HAlign_Left)
-				.AutoWidth()
+				SNew(SVerticalBox)
+				+SVerticalBox::Slot()
+				.HAlign(HAlign_Fill)
+				.VAlign(VAlign_Top)
+				.AutoHeight()
 				[
-					SNew(SButton)
-					.ButtonStyle(&FAppStyle::Get().GetWidgetStyle<FButtonStyle>("SimpleButton"))
-					.ToolTipText(LOCTEXT("BuildSelection_DestinationExploreTooltip", "Explore the destination"))
-					.OnClicked(this, &SBuildSelection::ExploreDestination_OnClicked)
+					SNew(SHorizontalBox)
+					+SHorizontalBox::Slot()
+					.HAlign(HAlign_Left)
+					.VAlign(VAlign_Bottom)
+					.AutoWidth()
 					[
-						SNew(SImage)
-						.Image_Lambda([this]
-						{
-							return GetSelectedBuildType() == EBuildType::Oplog ? FAppStyle::Get().GetBrush("Zen.BrowserView") : FAppStyle::Get().GetBrush("Zen.FolderView");
-						})
-						.ColorAndOpacity(FSlateColor::UseForeground())
+						SNew(STextBlock)
+							.ColorAndOpacity(TitleColor)
+							.Font(TitleFont)
+							.Text(LOCTEXT("BuildSelection_Destination", "Destination"))
+					]
+					+SHorizontalBox::Slot()
+					.HAlign(HAlign_Left)
+					.AutoWidth()
+					[
+						SNew(SButton)
+						.ButtonStyle(&FAppStyle::Get().GetWidgetStyle<FButtonStyle>("SimpleButton"))
+						.ToolTipText(LOCTEXT("BuildSelection_DestinationExploreTooltip", "Explore the destination"))
+						.OnClicked(this, &SBuildSelection::ExploreDestination_OnClicked)
+						[
+							SNew(SImage)
+							.Image_Lambda([this]
+							{
+								return GetSelectedBuildType() == EBuildType::Oplog ? FAppStyle::Get().GetBrush("Zen.BrowserView") : FAppStyle::Get().GetBrush("Zen.FolderView");
+							})
+							.ColorAndOpacity(FSlateColor::UseForeground())
+						]
 					]
 				]
-			]
-			+SVerticalBox::Slot()
-			.HAlign(HAlign_Fill)
-			.VAlign(VAlign_Top)
-			[
-				GetBuildDestinationPanel()
-			]
-			+SVerticalBox::Slot()
-			.HAlign(HAlign_Right)
-			.VAlign(VAlign_Bottom)
-			.Padding(FMargin(0, RowMargin))
-			[
-				SNew(SButton)
-				.Text(LOCTEXT("BuildSelection_Download", "Download"))
-				.ToolTipText(LOCTEXT("BuildSelection_DownloadTooltip", "Start a download of the selected build for the selected platforms"))
-				.ButtonStyle(FAppStyle::Get(), "Button")
-				.IsEnabled_Lambda([this]()
-				{
-					const bool bDestinationValid = !GetEffectiveDestination().IsEmpty();
-					if (!bDestinationValid)
+				+SVerticalBox::Slot()
+				.HAlign(HAlign_Fill)
+				.VAlign(VAlign_Top)
+				[
+					GetBuildDestinationPanel()
+				]
+				+SVerticalBox::Slot()
+				.HAlign(HAlign_Right)
+				.VAlign(VAlign_Bottom)
+				.Padding(FMargin(0, RowMargin))
+				[
+					SNew(SButton)
+					.Text(LOCTEXT("BuildSelection_Download", "Download"))
+					.ToolTipText(LOCTEXT("BuildSelection_DownloadTooltip", "Start a download of the selected build for the selected platforms"))
+					.ButtonStyle(FAppStyle::Get(), "Button")
+					.IsEnabled_Lambda([this]()
 					{
-						return false;
-					}
+						const bool bDestinationValid = !GetEffectiveDestination().IsEmpty();
+						if (!bDestinationValid)
+						{
+							return false;
+						}
 
-					TArray<FBuildSelectionBuildGroupPtr> SelectedItems = BuildGroupListView->GetSelectedItems();
-					if (SelectedItems.Num() == 0)
-					{
-						return false;
-					}
+						TArray<FBuildSelectionBuildGroupPtr> SelectedItems = BuildGroupListView->GetSelectedItems();
+						if (SelectedItems.Num() == 0)
+						{
+							return false;
+						}
 
-					FName SelectedConfigurationName(NAME_None);
-					if (SelectedConfiguration && *SelectedConfiguration != TEXT("All"))
-					{
-						SelectedConfigurationName = FName(*SelectedConfiguration);
-					}
-					for (const FString& ArtifactForSelectedGroup : SelectedGroupSelectedArtifacts)
-					{
-						if (FArtifact* FoundArtifact = SelectedItems[0]->NamedArtifacts.Find(ArtifactForSelectedGroup))
+						FName SelectedConfigurationName(NAME_None);
+						if (SelectedConfiguration && *SelectedConfiguration != TEXT("All"))
 						{
-							if ((SelectedConfigurationName == NAME_None) || FoundArtifact->Configurations.Contains(SelectedConfigurationName))
+							SelectedConfigurationName = FName(*SelectedConfiguration);
+						}
+						for (const FString& ArtifactForSelectedGroup : SelectedGroupSelectedArtifacts)
+						{
+							if (FArtifact* FoundArtifact = SelectedItems[0]->NamedArtifacts.Find(ArtifactForSelectedGroup))
 							{
-								return true;
+								if ((SelectedConfigurationName == NAME_None) || FoundArtifact->Configurations.Contains(SelectedConfigurationName))
+								{
+									return true;
+								}
 							}
 						}
-					}
-					return false;
-				})
-				.OnClicked_Lambda([this]()
-				{
-					DownloadSelectedArtifactsInSelectedBuildGroup();
-					return FReply::Handled();
-				})
+						return false;
+					})
+					.OnClicked_Lambda([this]()
+					{
+						DownloadSelectedArtifactsInSelectedBuildGroup();
+						return FReply::Handled();
+					})
+				]
 			]
 		]
 	];
diff --git a/Engine/Source/Programs/BuildStorageTool/Private/BuildStorageTool.cpp b/Engine/Source/Programs/BuildStorageTool/Private/BuildStorageTool.cpp
--- a/Engine/Source/Programs/BuildStorageTool/Private/BuildStorageTool.cpp
+++ b/Engine/Source/Programs/BuildStorageTool/Private/BuildStorageTool.cpp
@@ -28,7 +28,6 @@
 #include "Widgets/Input/SButton.h"
 #include "Widgets/Input/SButton.h"
 #include "Widgets/Input/SCheckBox.h"
-#include "Widgets/Layout/SScrollBox.h"
 #include "Widgets/Notifications/SNotificationList.h"
 #include "Widgets/Text/STextBlock.h"
 #include "ZenServiceInstanceManager.h"
@@ -320,59 +319,53 @@ public:
 				SNew(SBorder)
 				.BorderImage(FCoreStyle::Get().GetBrush("ToolPanel.GroupBorder"))
 				[
-					SNew(SScrollBox)
-					.Orientation(EOrientation::Orient_Horizontal)
-					+ SScrollBox::Slot()
-					.FillSize(1.0f)
+					SNew(SVerticalBox)
+
+					// Menu
+					+ SVerticalBox::Slot()
+					.AutoHeight()
+					[
+						MakeMainMenu()
+					]
+
+					// Login panel
+					+ SVerticalBox::Slot()
+					.AutoHeight()
+					.Padding(0.0f, 10.0f, 5.0f, 0.0f)
+					.HAlign(HAlign_Fill)
+					.VAlign(VAlign_Top)
+					[
+						SNew(SBuildLogin)
+						.ZenServiceInstance(ZenServiceInstanceManager.ToSharedRef(), &UE::Zen::FServiceInstanceManager::GetZenServiceInstance)
+						.BuildServiceInstance(BuildServiceInstanceManager.ToSharedRef(), &UE::Zen::Build::FServiceInstanceManager::GetBuildServiceInstance)
+					]
+
+					// Selection panel
+					+ SVerticalBox::Slot()
+					.Padding(0.0f, 10.0f, 5.0f, 10.0f)
+					.HAlign(HAlign_Fill)
+					.VAlign(VAlign_Fill)
+					[
+						SNew(SBuildSelection)
+						.ZenServiceInstance(ZenServiceInstanceManager.ToSharedRef(), &UE::Zen::FServiceInstanceManager::GetZenServiceInstance)
+						.BuildServiceInstance(BuildServiceInstanceManager.ToSharedRef(), &UE::Zen::Build::FServiceInstanceManager::GetBuildServiceInstance)
+						.OnBuildTransferStarted_Lambda([this]
+							(UE::Zen::Build::FBuildServiceInstance::FBuildTransfer Transfer, FStringView Name, FStringView Platform)
+							{
+								BuildActivity->AddBuildTransfer(Transfer, Name, Platform);
+							})
+					]
+
+					// Activity panel
+					+ SVerticalBox::Slot()
+					.AutoHeight()
+					.Padding(0.0f, 10.0f, 5.0f, 0.0f)
+					.HAlign(HAlign_Fill)
+					.VAlign(VAlign_Bottom)
 					[
-						SNew(SVerticalBox)
-
-						// Menu
-						+ SVerticalBox::Slot()
-						.AutoHeight()
-						[
-							MakeMainMenu()
-						]
-
-						// Login panel
-						+ SVerticalBox::Slot()
-						.AutoHeight()
-						.Padding(0.0f, 10.0f, 5.0f, 0.0f)
-						.HAlign(HAlign_Fill)
-						.VAlign(VAlign_Top)
-						[
-							SNew(SBuildLogin)
-							.ZenServiceInstance(ZenServiceInstanceManager.ToSharedRef(), &UE::Zen::FServiceInstanceManager::GetZenServiceInstance)
-							.BuildServiceInstance(BuildServiceInstanceManager.ToSharedRef(), &UE::Zen::Build::FServiceInstanceManager::GetBuildServiceInstance)
-						]
-
-						// Selection panel
-						+ SVerticalBox::Slot()
-						.Padding(0.0f, 10.0f, 5.0f, 10.0f)
-						.HAlign(HAlign_Fill)
-						.VAlign(VAlign_Fill)
-						[
-							SNew(SBuildSelection)
-							.ZenServiceInstance(ZenServiceInstanceManager.ToSharedRef(), &UE::Zen::FServiceInstanceManager::GetZenServiceInstance)
-							.BuildServiceInstance(BuildServiceInstanceManager.ToSharedRef(), &UE::Zen::Build::FServiceInstanceManager::GetBuildServiceInstance)
-							.OnBuildTransferStarted_Lambda([this]
-								(UE::Zen::Build::FBuildServiceInstance::FBuildTransfer Transfer, FStringView Name, FStringView Platform)
-								{
-									BuildActivity->AddBuildTransfer(Transfer, Name, Platform);
-								})
-						]
-
-						// Activity panel
-						+ SVerticalBox::Slot()
-						.AutoHeight()
-						.Padding(0.0f, 10.0f, 5.0f, 0.0f)
-						.HAlign(HAlign_Fill)
-						.VAlign(VAlign_Bottom)
-						[
-							SAssignNew(BuildActivity, SBuildActivity)
-							.ZenServiceInstance(ZenServiceInstanceManager.ToSharedRef(), &UE::Zen::FServiceInstanceManager::GetZenServiceInstance)
-							.BuildServiceInstance(BuildServiceInstanceManager.ToSharedRef(), &UE::Zen::Build::FServiceInstanceManager::GetBuildServiceInstance)
-						]
+						SAssignNew(BuildActivity, SBuildActivity)
+						.ZenServiceInstance(ZenServiceInstanceManager.ToSharedRef(), &UE::Zen::FServiceInstanceManager::GetZenServiceInstance)
+						.BuildServiceInstance(BuildServiceInstanceManager.ToSharedRef(), &UE::Zen::Build::FServiceInstanceManager::GetBuildServiceInstance)
 					]
 				]
 			];
