BuildStorageTool: Change scrollbar behaviore when we have a lot of build artifacts on one group.

#Synced-CL 49059893

diff --git a/Engine/Source/Editor/StorageServerWidgets/Private/SBuildSelection.cpp b/Engine/Source/Editor/StorageServerWidgets/Private/SBuildSelection.cpp
--- a/Engine/Source/Editor/StorageServerWidgets/Private/SBuildSelection.cpp
+++ b/Engine/Source/Editor/StorageServerWidgets/Private/SBuildSelection.cpp
@@ -2344,20 +2344,18 @@ TSharedRef<SWidget> SBuildSelection::GetGridPanel()
 	.HAlign(HAlign_Fill)
 	.VAlign(VAlign_Bottom)
 	[
-
-		SNew(SScrollBox)
-		.Orientation(EOrientation::Orient_Horizontal)
-		+ SScrollBox::Slot()
-		.FillSize(1.0f)
+		SNew(SHorizontalBox)
+		.Visibility_Lambda([this]()
+		{
+			return !!BuildListRefreshesInProgress || BuildGroupListView->GetNumItemsSelected() == 0 ? EVisibility::Collapsed : EVisibility::Visible;
+		})
+		+SHorizontalBox::Slot()
+		.HAlign(HAlign_Left)
 		[
-			SNew(SHorizontalBox)
-			.Visibility_Lambda([this]()
-			{
-				return !!BuildListRefreshesInProgress || BuildGroupListView->GetNumItemsSelected() == 0 ? EVisibility::Collapsed : EVisibility::Visible;
-			})
-			+SHorizontalBox::Slot()
-			.HAlign(HAlign_Left)
-			.AutoWidth()
+			SNew(SScrollBox)
+			.Orientation(EOrientation::Orient_Horizontal)
+			+ SScrollBox::Slot()
+			.FillSize(1.0f)
 			[
 				SNew(SVerticalBox)
 				+SVerticalBox::Slot()
@@ -2392,97 +2390,97 @@ TSharedRef<SWidget> SBuildSelection::GetGridPanel()
 					SAssignNew(SelectedGroupArtifactGrid, SGridPanel)
 				]
 			]
-			+SHorizontalBox::Slot()
-			.HAlign(HAlign_Right)
-			.FillWidth(0.5f)
+		]
+		+SHorizontalBox::Slot()
+		.HAlign(HAlign_Right)
+		.AutoWidth()
+		[
+			SNew(SVerticalBox)
+			+SVerticalBox::Slot()
+			.HAlign(HAlign_Fill)
+			.VAlign(VAlign_Top)
+			.AutoHeight()
 			[
-				SNew(SVerticalBox)
-				+SVerticalBox::Slot()
-				.HAlign(HAlign_Fill)
-				.VAlign(VAlign_Top)
-				.AutoHeight()
-				[
-					SNew(SHorizontalBox)
-					+SHorizontalBox::Slot()
-					.HAlign(HAlign_Left)
-					.VAlign(VAlign_Bottom)
-					.AutoWidth()
-					[
-						SNew(STextBlock)
-							.ColorAndOpacity(TitleColor)
-							.Font(TitleFont)
-							.Text(LOCTEXT("BuildSelection_Destination", "Destination"))
-					]
-					+SHorizontalBox::Slot()
-					.HAlign(HAlign_Left)
-					.AutoWidth()
-					[
-						SNew(SButton)
-						.ButtonStyle(&FAppStyle::Get().GetWidgetStyle<FButtonStyle>("SimpleButton"))
-						.ToolTipText(LOCTEXT("BuildSelection_DestinationExploreTooltip", "Explore the destination"))
-						.OnClicked(this, &SBuildSelection::ExploreDestination_OnClicked)
-						[
-							SNew(SImage)
-							.Image_Lambda([this]
-							{
-								return GetSelectedBuildType() == EBuildType::Oplog ? FAppStyle::Get().GetBrush("Zen.BrowserView") : FAppStyle::Get().GetBrush("Zen.FolderView");
-							})
-							.ColorAndOpacity(FSlateColor::UseForeground())
-						]
-					]
-				]
-				+SVerticalBox::Slot()
-				.HAlign(HAlign_Fill)
-				.VAlign(VAlign_Top)
+				SNew(SHorizontalBox)
+				+SHorizontalBox::Slot()
+				.HAlign(HAlign_Left)
+				.VAlign(VAlign_Bottom)
+				.AutoWidth()
 				[
-					GetBuildDestinationPanel()
+					SNew(STextBlock)
+						.ColorAndOpacity(TitleColor)
+						.Font(TitleFont)
+						.Text(LOCTEXT("BuildSelection_Destination", "Destination"))
 				]
-				+SVerticalBox::Slot()
-				.HAlign(HAlign_Right)
-				.VAlign(VAlign_Bottom)
-				.Padding(FMargin(0, RowMargin))
+				+SHorizontalBox::Slot()
+				.HAlign(HAlign_Left)
+				.AutoWidth()
 				[
 					SNew(SButton)
-					.Text(LOCTEXT("BuildSelection_Download", "Download"))
-					.ToolTipText(LOCTEXT("BuildSelection_DownloadTooltip", "Start a download of the selected build for the selected platforms"))
-					.ButtonStyle(FAppStyle::Get(), "Button")
-					.IsEnabled_Lambda([this]()
-					{
-						const bool bDestinationValid = !GetEffectiveDestination().IsEmpty();
-						if (!bDestinationValid)
+					.ButtonStyle(&FAppStyle::Get().GetWidgetStyle<FButtonStyle>("SimpleButton"))
+					.ToolTipText(LOCTEXT("BuildSelection_DestinationExploreTooltip", "Explore the destination"))
+					.OnClicked(this, &SBuildSelection::ExploreDestination_OnClicked)
+					[
+						SNew(SImage)
+						.Image_Lambda([this]
 						{
-							return false;
-						}
+							return GetSelectedBuildType() == EBuildType::Oplog ? FAppStyle::Get().GetBrush("Zen.BrowserView") : FAppStyle::Get().GetBrush("Zen.FolderView");
+						})
+						.ColorAndOpacity(FSlateColor::UseForeground())
+					]
+				]
+			]
+			+SVerticalBox::Slot()
+			.HAlign(HAlign_Fill)
+			.VAlign(VAlign_Top)
+			[
+				GetBuildDestinationPanel()
+			]
+			+SVerticalBox::Slot()
+			.HAlign(HAlign_Right)
+			.VAlign(VAlign_Bottom)
+			.Padding(FMargin(0, RowMargin))
+			[
+				SNew(SButton)
+				.Text(LOCTEXT("BuildSelection_Download", "Download"))
+				.ToolTipText(LOCTEXT("BuildSelection_DownloadTooltip", "Start a download of the selected build for the selected platforms"))
+				.ButtonStyle(FAppStyle::Get(), "Button")
+				.IsEnabled_Lambda([this]()
+				{
+					const bool bDestinationValid = !GetEffectiveDestination().IsEmpty();
+					if (!bDestinationValid)
+					{
+						return false;
+					}
 
-						TArray<FBuildSelectionBuildGroupPtr> SelectedItems = BuildGroupListView->GetSelectedItems();
-						if (SelectedItems.Num() == 0)
-						{
-							return false;
-						}
+					TArray<FBuildSelectionBuildGroupPtr> SelectedItems = BuildGroupListView->GetSelectedItems();
+					if (SelectedItems.Num() == 0)
+					{
+						return false;
+					}
 
-						FName SelectedConfigurationName(NAME_None);
-						if (SelectedConfiguration && *SelectedConfiguration != TEXT("All"))
-						{
-							SelectedConfigurationName = FName(*SelectedConfiguration);
-						}
-						for (const FString& ArtifactForSelectedGroup : SelectedGroupSelectedArtifacts)
+					FName SelectedConfigurationName(NAME_None);
+					if (SelectedConfiguration && *SelectedConfiguration != TEXT("All"))
+					{
+						SelectedConfigurationName = FName(*SelectedConfiguration);
+					}
+					for (const FString& ArtifactForSelectedGroup : SelectedGroupSelectedArtifacts)
+					{
+						if (FArtifact* FoundArtifact = SelectedItems[0]->NamedArtifacts.Find(ArtifactForSelectedGroup))
 						{
-							if (FArtifact* FoundArtifact = SelectedItems[0]->NamedArtifacts.Find(ArtifactForSelectedGroup))
+							if ((SelectedConfigurationName == NAME_None) || FoundArtifact->Configurations.Contains(SelectedConfigurationName))
 							{
-								if ((SelectedConfigurationName == NAME_None) || FoundArtifact->Configurations.Contains(SelectedConfigurationName))
-								{
-									return true;
-								}
+								return true;
 							}
 						}
-						return false;
-					})
-					.OnClicked_Lambda([this]()
-					{
-						DownloadSelectedArtifactsInSelectedBuildGroup();
-						return FReply::Handled();
-					})
-				]
+					}
+					return false;
+				})
+				.OnClicked_Lambda([this]()
+				{
+					DownloadSelectedArtifactsInSelectedBuildGroup();
+					return FReply::Handled();
+				})
 			]
 		]
 	];
