BuildStorageTool: Fix bug where SelectedConfiguration was influencing ability to download builds even when you were switched to a build type that doesn't have a configuration associated with artifacts.

#Synced-CL 49059893

diff --git a/Engine/Source/Editor/StorageServerWidgets/Private/SBuildSelection.cpp b/Engine/Source/Editor/StorageServerWidgets/Private/SBuildSelection.cpp
--- a/Engine/Source/Editor/StorageServerWidgets/Private/SBuildSelection.cpp
+++ b/Engine/Source/Editor/StorageServerWidgets/Private/SBuildSelection.cpp
@@ -1597,7 +1597,7 @@ void SBuildSelection::DownloadSelectedArtifactsInSelectedBuildGroup()
 		}
 
 		FName SelectedConfigurationName(NAME_None);
-		if (SelectedConfiguration && *SelectedConfiguration != TEXT("All"))
+		if (!ArtifactConfigurationSet.IsEmpty() && SelectedConfiguration && *SelectedConfiguration != TEXT("All"))
 		{
 			SelectedConfigurationName = FName(*SelectedConfiguration);
 		}
@@ -2460,7 +2460,7 @@ TSharedRef<SWidget> SBuildSelection::GetGridPanel()
 					}
 
 					FName SelectedConfigurationName(NAME_None);
-					if (SelectedConfiguration && *SelectedConfiguration != TEXT("All"))
+					if (!ArtifactConfigurationSet.IsEmpty() && SelectedConfiguration && *SelectedConfiguration != TEXT("All"))
 					{
 						SelectedConfigurationName = FName(*SelectedConfiguration);
 					}
