Fix null dereference crash in BuildStorageTool if SelectedProject is null.

#Synced-CL 50439800

diff --git a/Engine/Source/Editor/StorageServerWidgets/Private/SBuildSelection.cpp b/Engine/Source/Editor/StorageServerWidgets/Private/SBuildSelection.cpp
--- a/Engine/Source/Editor/StorageServerWidgets/Private/SBuildSelection.cpp
+++ b/Engine/Source/Editor/StorageServerWidgets/Private/SBuildSelection.cpp
@@ -1299,24 +1299,34 @@ FString SBuildSelection::GetDefaultDestination() const
 	EBuildType BuildType = GetSelectedBuildType();
 	IPlatformFile& PlatformFile = FPlatformFileManager::Get().GetPlatformFile();
 
-	FString FallbackPath = FPaths::EngineSavedDir();
-	FString ProjectFilename = FUProjectDictionary::GetDefault().GetProjectPathForGame(**SelectedProject);
-
-	if (ProjectFilename.IsEmpty() && UserSelectedProjectDictionary)
+	FString ProjectFilename;
+	if (SelectedProject.IsValid())
 	{
-		ProjectFilename = UserSelectedProjectDictionary->GetProjectPathForGame(**SelectedProject);
-	}
+		ProjectFilename = FUProjectDictionary::GetDefault().GetProjectPathForGame(**SelectedProject);
 
-	if (ProjectFilename.IsEmpty() && FallbackProjectDictionary)
-	{
-		ProjectFilename = FallbackProjectDictionary->GetProjectPathForGame(**SelectedProject);
+		if (ProjectFilename.IsEmpty() && UserSelectedProjectDictionary)
+		{
+			ProjectFilename = UserSelectedProjectDictionary->GetProjectPathForGame(**SelectedProject);
+		}
+
+		if (ProjectFilename.IsEmpty() && FallbackProjectDictionary)
+		{
+			ProjectFilename = FallbackProjectDictionary->GetProjectPathForGame(**SelectedProject);
+		}
 	}
 
 	if (ProjectFilename.IsEmpty())
 	{
 		if (BuildType == EBuildType::Oplog)
 		{
-			return *SelectedProject;
+			if (SelectedProject.IsValid())
+			{
+				return *SelectedProject;
+			}
+			else
+			{
+				return TEXT("DownloadedBuilds");
+			}
 		}
 
 		FString DownloadBase;
