BuildStorageTool: Fix scrollwheel scaling bug by not doing list item filtering dynamically and instead binding a filtered list to the list view.

#Synced-CL 51039791

diff --git a/Engine/Source/Editor/StorageServerWidgets/Internal/SBuildSelection.h b/Engine/Source/Editor/StorageServerWidgets/Internal/SBuildSelection.h
--- a/Engine/Source/Editor/StorageServerWidgets/Internal/SBuildSelection.h
+++ b/Engine/Source/Editor/StorageServerWidgets/Internal/SBuildSelection.h
@@ -92,7 +92,7 @@ private:
 	EBuildType GetSelectedBuildType() const;
 	void RebuildLists();
 	void RegenerateBuildGroups(UE::Zen::Build::FListBuildsState& ListBuildsState);
-	void ConditionalSortBuildGroups();
+	void RefreshPresentedBuildGroups();
 	void RegenerateActivePlatformFilters();
 	void ValidateBuildGroupSelection();
 	TSharedRef<SWidget> OnGenerateTextBlockFromString(TSharedPtr<FString> Item);
@@ -212,6 +212,7 @@ private:
 	std::atomic<uint32> BuildListRefreshesInProgress = 0;
 	std::atomic<uint32> BuildRefreshGeneration = 0;
 	TArray<TSharedPtr<SBuildSelection::FBuildGroup, ESPMode::ThreadSafe>> BuildGroups;
+	TArray<TSharedPtr<SBuildSelection::FBuildGroup, ESPMode::ThreadSafe>> PresentedBuildGroups;
 
 	TAttribute<TSharedPtr<UE::Zen::FZenServiceInstance>> ZenServiceInstance;
 	TAttribute<TSharedPtr<UE::Zen::Build::FBuildServiceInstance>> BuildServiceInstance;
diff --git a/Engine/Source/Editor/StorageServerWidgets/Private/SBuildSelection.cpp b/Engine/Source/Editor/StorageServerWidgets/Private/SBuildSelection.cpp
--- a/Engine/Source/Editor/StorageServerWidgets/Private/SBuildSelection.cpp
+++ b/Engine/Source/Editor/StorageServerWidgets/Private/SBuildSelection.cpp
@@ -438,6 +438,7 @@ void SBuildSelection::RebuildLists()
 
 		++BuildRefreshGeneration;
 		BuildGroups.Empty();
+		PresentedBuildGroups.Reset();
 
 		for (const FString& Namespace : NamespacesToQuery)
 		{
@@ -496,6 +497,7 @@ void SBuildSelection::RegenerateBuildGroups(UE::Zen::Build::FListBuildsState& Li
 	using namespace UE::BuildSelection::Internal;
 
 	BuildGroups.Empty();
+	PresentedBuildGroups.Reset();
 	ArtifactNameMap.Empty();
 	ArtifactConfigurationSet.Empty();
 	ArtifactConfigurationList.Empty();
@@ -680,13 +682,22 @@ void SBuildSelection::RegenerateBuildGroups(UE::Zen::Build::FListBuildsState& Li
 		WriteSetting(TEXT("Configuration"), SelectedConfiguration ? **SelectedConfiguration : TEXT(""));
 	}
 	KeyedGroups.GenerateValueArray(BuildGroups);
-	ConditionalSortBuildGroups();
+	RefreshPresentedBuildGroups();
 }
 
-void SBuildSelection::ConditionalSortBuildGroups()
+void SBuildSelection::RefreshPresentedBuildGroups()
 {
 	using namespace UE::BuildSelection::Internal;
 
+	PresentedBuildGroups.Reset();
+	for (TSharedPtr<FBuildGroup> BuildGroup : BuildGroups)
+	{
+		if (DoesPassAllViewFilters(BuildGroup))
+		{
+			PresentedBuildGroups.Add(BuildGroup);
+		}
+	}
+
 	if (BuildGroupSortMode == EColumnSortMode::None)
 	{
 		return;
@@ -697,25 +708,25 @@ void SBuildSelection::ConditionalSortBuildGroups()
 	{
 		if (BuildGroupSortByColumn == FBuildGroupIds::ColName)
 		{
-			BuildGroups.Sort([this](const TSharedPtr<FBuildGroup>& A, const TSharedPtr<FBuildGroup>& B) {
+			PresentedBuildGroups.Sort([this](const TSharedPtr<FBuildGroup>& A, const TSharedPtr<FBuildGroup>& B) {
 				return (A->DisplayName < B->DisplayName);
 			});
 		}
 		else if (BuildGroupSortByColumn == FBuildGroupIds::ColCommit)
 		{
-			BuildGroups.Sort([this](const TSharedPtr<FBuildGroup>& A, const TSharedPtr<FBuildGroup>& B) {
+			PresentedBuildGroups.Sort([this](const TSharedPtr<FBuildGroup>& A, const TSharedPtr<FBuildGroup>& B) {
 				return (A->CommitIdentifier < B->CommitIdentifier);
 			});
 		}
 		else if (BuildGroupSortByColumn == FBuildGroupIds::ColCategory)
 		{
-			BuildGroups.Sort([this](const TSharedPtr<FBuildGroup>& A, const TSharedPtr<FBuildGroup>& B) {
+			PresentedBuildGroups.Sort([this](const TSharedPtr<FBuildGroup>& A, const TSharedPtr<FBuildGroup>& B) {
 				return (A->Category < B->Category);
 			});
 		}
 		else if (BuildGroupSortByColumn == FBuildGroupIds::ColCreated)
 		{
-			BuildGroups.Sort([this](const TSharedPtr<FBuildGroup>& A, const TSharedPtr<FBuildGroup>& B) {
+			PresentedBuildGroups.Sort([this](const TSharedPtr<FBuildGroup>& A, const TSharedPtr<FBuildGroup>& B) {
 				return (A->CreatedAt < B->CreatedAt);
 			});
 		}
@@ -724,26 +735,26 @@ void SBuildSelection::ConditionalSortBuildGroups()
 	{
 		if (BuildGroupSortByColumn == FBuildGroupIds::ColName)
 		{
-			BuildGroups.Sort([this](const TSharedPtr<FBuildGroup>& A, const TSharedPtr<FBuildGroup>& B) {
-				return (A->DisplayName >= B->DisplayName);
+			PresentedBuildGroups.Sort([this](const TSharedPtr<FBuildGroup>& A, const TSharedPtr<FBuildGroup>& B) {
+				return (B->DisplayName < A->DisplayName);
 			});
 		}
 		else if (BuildGroupSortByColumn == FBuildGroupIds::ColCommit)
 		{
-			BuildGroups.Sort([this](const TSharedPtr<FBuildGroup>& A, const TSharedPtr<FBuildGroup>& B) {
-				return (A->CommitIdentifier >= B->CommitIdentifier);
+			PresentedBuildGroups.Sort([this](const TSharedPtr<FBuildGroup>& A, const TSharedPtr<FBuildGroup>& B) {
+				return (B->CommitIdentifier < A->CommitIdentifier);
 			});
 		}
 		else if (BuildGroupSortByColumn == FBuildGroupIds::ColCategory)
 		{
-			BuildGroups.Sort([this](const TSharedPtr<FBuildGroup>& A, const TSharedPtr<FBuildGroup>& B) {
-				return (A->Category >= B->Category);
+			PresentedBuildGroups.Sort([this](const TSharedPtr<FBuildGroup>& A, const TSharedPtr<FBuildGroup>& B) {
+				return (B->Category < A->Category);
 			});
 		}
 		else if (BuildGroupSortByColumn == FBuildGroupIds::ColCreated)
 		{
-			BuildGroups.Sort([this](const TSharedPtr<FBuildGroup>& A, const TSharedPtr<FBuildGroup>& B) {
-				return (A->CreatedAt >= B->CreatedAt);
+			PresentedBuildGroups.Sort([this](const TSharedPtr<FBuildGroup>& A, const TSharedPtr<FBuildGroup>& B) {
+				return (B->CreatedAt < A->CreatedAt);
 			});
 		}
 	}
@@ -849,6 +860,7 @@ bool SBuildSelection::DoesPassAllViewFilters(TSharedPtr<FBuildGroup> InItem) con
 void SBuildSelection::PostViewFilterChange()
 {
 	ValidateBuildGroupSelection();
+	RefreshPresentedBuildGroups();
 	BuildGroupListView->Invalidate(EInvalidateWidgetReason::LayoutAndVolatility);
 	BuildGroupListView->RequestListRefresh();
 }
@@ -860,15 +872,7 @@ bool SBuildSelection::BuildGroupIsSelectableOrNavigable(FBuildSelectionBuildGrou
 
 TSharedRef<ITableRow> SBuildSelection::GenerateBuildGroupRow(FBuildSelectionBuildGroupPtr InItem, const TSharedRef<STableViewBase>& InOwningTable)
 {
-	return SNew(SBuildGroupTableRow, InOwningTable, InItem)
-		.Visibility_Lambda([this, InItem]()
-		{
-			if (DoesPassAllViewFilters(InItem))
-			{
-				return EVisibility::Visible;
-			}
-			return EVisibility::Collapsed;
-		});
+	return SNew(SBuildGroupTableRow, InOwningTable, InItem);
 }
 
 TSharedPtr<SWidget> SBuildSelection::OnGetBuildGroupContextMenuContent() const
@@ -1206,7 +1210,7 @@ void SBuildSelection::OnColumnSortModeChanged(const EColumnSortPriority::Type So
 	BuildGroupSortByColumn = ColumnId;
 	BuildGroupSortMode = InSortMode;
 
-	ConditionalSortBuildGroups();
+	RefreshPresentedBuildGroups();
 	ExecuteOnGameThread(UE_SOURCE_LOCATION,
 		[this]
 		{
@@ -2314,7 +2318,7 @@ TSharedRef<SWidget> SBuildSelection::GetGridPanel()
 		.VAlign(VAlign_Fill)
 		[
 			SAssignNew(BuildGroupListView, SListView<FBuildSelectionBuildGroupPtr>)
-				.ListItemsSource(&BuildGroups)
+				.ListItemsSource(&PresentedBuildGroups)
 				.OnGenerateRow(this, &SBuildSelection::GenerateBuildGroupRow)
 				.OnSelectionChanged(this, &SBuildSelection::BuildGroupSelectionChanged)
 				.OnContextMenuOpening(this, &SBuildSelection::OnGetBuildGroupContextMenuContent)
@@ -2364,15 +2368,7 @@ TSharedRef<SWidget> SBuildSelection::GetGridPanel()
 				{
 					return LOCTEXT("BuildSelection_ResultLoading", "Loading...");
 				}
-				int32 VisibleItemCount = 0;
-				for (FBuildSelectionBuildGroupPtr BuildGroup : BuildGroups)
-				{
-					if (DoesPassAllViewFilters(BuildGroup))
-					{
-						VisibleItemCount++;
-					}
-				}
-				return FText::Format(LOCTEXT("BuildSelection_ResultDescription", "{0} {0}|plural(one=item,other=items)"), FText::AsNumber(VisibleItemCount));
+				return FText::Format(LOCTEXT("BuildSelection_ResultDescription", "{0} {0}|plural(one=item,other=items)"), FText::AsNumber(PresentedBuildGroups.Num()));
 			})
 	];
 
