BuildStorageTool: Extract the artifact selection widget from build group selection in preparation for some rework of the artifact selection widget contents.

#Synced-CL 51039791

diff --git a/Engine/Source/Editor/StorageServerWidgets/Internal/BuildSelectionInternal.h b/Engine/Source/Editor/StorageServerWidgets/Internal/BuildSelectionInternal.h
new file mode 100644
--- /dev/null
+++ b/Engine/Source/Editor/StorageServerWidgets/Internal/BuildSelectionInternal.h
@@ -0,0 +1,25 @@
+// Copyright Epic Games, Inc. All Rights Reserved.
+
+#pragma once
+
+#include "Containers/Set.h"
+#include "Containers/StringFwd.h"
+#include "Experimental/BuildServerInterface.h"
+
+namespace UE::BuildSelection::Internal
+{
+
+struct FArtifact
+{
+	TSet<FName> Configurations;
+	UE::Zen::Build::FBuildServiceInstance::FBuildRecord BuildRecord;
+};
+
+void ConformSelection(TSharedPtr<FString>& SelectedItem, const TArray<TSharedPtr<FString>>& SelectionList);
+
+bool WriteSetting(FStringView InSectionName, FStringView InKeyName, FStringView InValue);
+bool WriteSetting(FStringView InKeyName, FStringView InValue);
+bool ReadSetting(FStringView InSectionName, FStringView InKeyName, FString& OutValue);
+bool ReadSetting(FStringView InKeyName, FString& OutValue);
+
+} // namespace UE::Zen::Build
\ No newline at end of file
diff --git a/Engine/Source/Editor/StorageServerWidgets/Internal/SBuildArtifactSelection.h b/Engine/Source/Editor/StorageServerWidgets/Internal/SBuildArtifactSelection.h
new file mode 100644
--- /dev/null
+++ b/Engine/Source/Editor/StorageServerWidgets/Internal/SBuildArtifactSelection.h
@@ -0,0 +1,62 @@
+// Copyright Epic Games, Inc. All Rights Reserved.
+
+#pragma once
+
+#include "BuildSelectionInternal.h"
+#include "Experimental/BuildServerInterface.h"
+#include "Experimental/ZenServerInterface.h"
+#include "Templates/SharedPointer.h"
+#include "Widgets/DeclarativeSyntaxSupport.h"
+#include "Widgets/Layout/SGridPanel.h"
+#include "Widgets/SCompoundWidget.h"
+#include "Widgets/SWidget.h"
+#include "Widgets/Views/ITableRow.h"
+#include "Widgets/Views/SListView.h"
+#include "Widgets/Views/STableRow.h"
+#include "Widgets/Views/STableViewBase.h"
+#include "ZenBuildUtils.h"
+
+#define UE_API STORAGESERVERWIDGETS_API
+
+class SBuildArtifactSelection : public SCompoundWidget
+{
+public:
+	SLATE_BEGIN_ARGS(SBuildArtifactSelection)
+		: _ZenServiceInstance(nullptr)
+		, _BuildServiceInstance(nullptr)
+
+	{ }
+
+	SLATE_ATTRIBUTE(TSharedPtr<UE::Zen::FZenServiceInstance>, ZenServiceInstance);
+	SLATE_ATTRIBUTE(TSharedPtr<UE::Zen::Build::FBuildServiceInstance>, BuildServiceInstance);
+
+	SLATE_END_ARGS()
+
+	UE_API void Construct(const FArguments& InArgs);
+
+	bool HasArtifactConfigurations() const;
+	TSharedPtr<FString> GetSelectedConfiguration() const;
+	const TArray<FString>& GetSelectedArtifacts() const;
+	UE_API void SetAvailableConfigurations(
+			TSet<FString>&& InArtifactConfigurationSet,
+			TMultiMap<FString, FString>&& InArtifactNameMap
+		);
+	UE_API void Refresh(const TMap<FString, UE::BuildSelection::Internal::FArtifact>* InNamedArtifacts);
+private:
+	void RebuildSelectedGroupArtifactHeader();
+	void RebuildSelectedGroupArtifactContents();
+	void SetSelectedConfiguration(const TSharedPtr<FString> InSelectedConfiguration);
+
+	TMultiMap<FString, FString> ArtifactNameMap;
+	TSet<FString> ArtifactConfigurationSet;
+	TArray<TSharedPtr<FString>> ArtifactConfigurationList;
+	TSharedPtr<FString> SelectedConfiguration;
+	TArray<FString> SelectedArtifacts;
+
+	TMap<FString, UE::BuildSelection::Internal::FArtifact> NamedArtifacts;
+
+	TSharedPtr<SGridPanel> SelectedGroupArtifactHeaderGrid;
+	TSharedPtr<SGridPanel> SelectedGroupArtifactGrid;
+};
+
+#undef UE_API
diff --git a/Engine/Source/Editor/StorageServerWidgets/Internal/SBuildSelection.h b/Engine/Source/Editor/StorageServerWidgets/Internal/SBuildSelection.h
--- a/Engine/Source/Editor/StorageServerWidgets/Internal/SBuildSelection.h
+++ b/Engine/Source/Editor/StorageServerWidgets/Internal/SBuildSelection.h
@@ -2,6 +2,7 @@
 
 #pragma once
 
+#include "BuildSelectionInternal.h"
 #include "Experimental/BuildServerInterface.h"
 #include "Experimental/ZenServerInterface.h"
 #include "Internationalization/Regex.h"
@@ -21,6 +22,7 @@
 #define UE_API STORAGESERVERWIDGETS_API
 
 class FUProjectDictionary;
+class SBuildArtifactSelection;
 class SMultiSelectComboBox;
 
 namespace UE::Zen::Build
@@ -51,12 +53,6 @@ class SBuildSelection : public SCompoundWidget
 	SLATE_EVENT(FOnBuildTransferStarted, OnBuildTransferStarted)
 	SLATE_END_ARGS()
 
-	struct FArtifact
-	{
-		TSet<FName> Configurations;
-		UE::Zen::Build::FBuildServiceInstance::FBuildRecord BuildRecord;
-	};
-
 	struct FBuildGroup : public TSharedFromThis<FBuildGroup, ESPMode::ThreadSafe>
 	{
 		FString Namespace;
@@ -65,7 +61,7 @@ class SBuildSelection : public SCompoundWidget
 		FDateTime CreatedAt;
 		FString Category;
 		FString Job;
-		TMap<FString, FArtifact> NamedArtifacts;
+		TMap<FString, UE::BuildSelection::Internal::FArtifact> NamedArtifacts;
 		bool bIsPreflight;
 	};
 
@@ -106,8 +102,7 @@ private:
 	void OnChooseTransferFilterClicked();
 	TSharedRef<SWidget> GetBuildDestinationPanel();
 	TSharedRef<SWidget> GetGridPanel();
-	void RebuildSelectedGroupArtifactHeader(TSharedPtr<FBuildGroup> Item);
-	void RebuildSelectedGroupArtifactGrid(TSharedPtr<FBuildGroup> Item);
+	void RefreshArtifactSelectionWidget(TSharedPtr<FBuildGroup> Item);
 	void BuildGroupSelectionChanged(TSharedPtr<FBuildGroup> Item, ESelectInfo::Type SelectInfo);
 	EColumnSortMode::Type GetColumnSortMode(const FName ColumnId) const;
 	void OnColumnSortModeChanged(const EColumnSortPriority::Type SortPriority, const FName& ColumnId, const EColumnSortMode::Type InSortMode);
@@ -127,12 +122,6 @@ private:
 	void SetSelectedProject(const TSharedPtr<FString> InSelectedProject);
 	void SetSelectedBuildType(const TSharedPtr<FString> InBuildType);
 	void SetIncludePreflights(const bool bInIncludePreflights);
-	void SetSelectedConfiguration(const TSharedPtr<FString> InSelectedStream);
-
-	bool WriteSetting(FStringView InSectionName, FStringView InKeyName, FStringView InValue);
-	bool WriteSetting(FStringView InKeyName, FStringView InValue);
-	bool ReadSetting(FStringView InSectionName, FStringView InKeyName, FString& OutValue);
-	bool ReadSetting(FStringView InKeyName, FString& OutValue);
 
 	void SetUserSelectedProjectDictionaryRoot(FStringView InRoot);
 	void SetFallbackProjectDictionaryRoot(FStringView InRoot);
@@ -142,16 +131,11 @@ private:
 	static FString SanitizeBranch(const FString& InString);
 	static FString SanitizeBucket(const FString& InString);
 
-	static void ConformSelection(TSharedPtr<FString>& SelectedItem, const TArray<TSharedPtr<FString>>& SelectionList);
-
 	static const TCHAR* LexToString(EBuildType BuildType);
 
 	FReply ExploreDestination_OnClicked();
-
 	void DownloadSelectedArtifactsInSelectedBuildGroup();
 
-	static const FString EmptyString;
-
 	TMap<FString, FString> EngineInstallations;
 
 	FString FallbackProjectDictionaryRootDir;
@@ -171,19 +155,13 @@ private:
 	TSharedPtr<SComboBox<TSharedPtr<FString>>> BuildTypeWidget;
 	TSharedPtr<SMultiSelectComboBox> PlatformFiltersWidget;
 
-	TSharedPtr<SGridPanel> SelectedGroupArtifactHeaderGrid;
-	TSharedPtr<SGridPanel> SelectedGroupArtifactGrid;
-	TArray<FString> SelectedGroupSelectedArtifacts;
+	TSharedPtr<SBuildArtifactSelection> BuildArtifactSelection;
 
 	TArray<TSharedPtr<FString>> ProjectList;
 	TArray<TSharedPtr<FString>> StreamList;
 	TArray<TSharedPtr<FString>> BuildTypeList;
 	TArray<TSharedPtr<FString>> PlatformList;
 
-	TMultiMap<FString, FString> ArtifactNameMap;
-	TSet<FString> ArtifactConfigurationSet;
-	TArray<TSharedPtr<FString>> ArtifactConfigurationList;
-
 	struct FTransferFilter
 	{
 		FString IncludeFilter;
@@ -199,7 +177,6 @@ private:
 	TSharedPtr<FString> SelectedStream;
 	TSharedPtr<FString> SelectedProject;
 	TSharedPtr<FString> SelectedBuildType;
-	TSharedPtr<FString> SelectedConfiguration;
 	FString SelectedTextFilter;
 	bool bIncludePreflights = true;
 
diff --git a/Engine/Source/Editor/StorageServerWidgets/Private/BuildSelectionInternal.cpp b/Engine/Source/Editor/StorageServerWidgets/Private/BuildSelectionInternal.cpp
new file mode 100644
--- /dev/null
+++ b/Engine/Source/Editor/StorageServerWidgets/Private/BuildSelectionInternal.cpp
@@ -0,0 +1,52 @@
+// Copyright Epic Games, Inc. All Rights Reserved.
+
+#include "BuildSelectionInternal.h"
+#include "Misc/Paths.h"
+
+namespace UE::BuildSelection::Internal
+{
+
+void ConformSelection(TSharedPtr<FString>& SelectedItem, const TArray<TSharedPtr<FString>>& SelectionList)
+{
+	if (!SelectedItem)
+	{
+		SelectedItem = SelectionList.IsEmpty() ? nullptr : SelectionList[0];
+		return;
+	}
+
+	const TSharedPtr<FString>* FoundSelectionListItem = SelectionList.FindByPredicate([&SelectedItem](const TSharedPtr<FString>& Item)
+		{
+			return *Item == *SelectedItem;
+		});
+
+	SelectedItem = FoundSelectionListItem ? *FoundSelectionListItem : SelectionList.IsEmpty() ? nullptr : SelectionList[0];
+}
+
+bool WriteSetting(const FStringView InSectionName, const FStringView InKeyName, const FStringView InValue)
+{
+	static const FString SectionName = TEXT("Unreal Engine/Build Storage");
+	const FString FinalSectionName = InSectionName.IsEmpty() ? SectionName : FPaths::Combine(SectionName, InSectionName);
+	const FString FinalKeyName(InKeyName);
+	FString Value(InValue);
+	return FPlatformMisc::SetStoredValue(TEXT("Epic Games"), FinalSectionName, FinalKeyName, Value);
+}
+
+bool WriteSetting(const FStringView InKeyName, const FStringView InValue)
+{
+	return WriteSetting(TEXT(""), InKeyName, InValue);
+}
+
+bool ReadSetting(const FStringView InSectionName, const FStringView InKeyName, FString& OutValue)
+{
+	static const FString SectionName = TEXT("Unreal Engine/Build Storage");
+	const FString FinalSectionName = InSectionName.IsEmpty() ? SectionName : FPaths::Combine(SectionName, InSectionName);
+	const FString FinalKeyName(InKeyName);
+	return FPlatformMisc::GetStoredValue(TEXT("Epic Games"), FinalSectionName, FinalKeyName, OutValue);
+}
+
+bool ReadSetting(const FStringView InKeyName, FString& OutValue)
+{
+	return ReadSetting(TEXT(""), InKeyName, OutValue);
+}
+
+} // namespace UE::Zen::Build
diff --git a/Engine/Source/Editor/StorageServerWidgets/Private/SBuildArtifactSelection.cpp b/Engine/Source/Editor/StorageServerWidgets/Private/SBuildArtifactSelection.cpp
new file mode 100644
--- /dev/null
+++ b/Engine/Source/Editor/StorageServerWidgets/Private/SBuildArtifactSelection.cpp
@@ -0,0 +1,390 @@
+// Copyright Epic Games, Inc. All Rights Reserved.
+
+#include "SBuildArtifactSelection.h"
+
+#include "Framework/MultiBox/MultiBoxBuilder.h"
+#include "Styling/StyleColors.h"
+#include "Widgets/Input/SButton.h"
+#include "Widgets/Input/SCheckBox.h"
+#include "Widgets/Input/SComboBox.h"
+#include "Widgets/Text/STextBlock.h"
+#include "ZenServiceInstanceManager.h"
+#include <atomic>
+
+#define LOCTEXT_NAMESPACE "StorageServerBuild"
+
+namespace UE::BuildArtifactSelection::Internal
+{
+	namespace FBuildArtifactIds
+	{
+		const FName ColSelected = TEXT("Selected");
+		const FName ColName = TEXT("Name");
+	}
+}
+
+void SBuildArtifactSelection::Construct(const FArguments& InArgs)
+{
+	using namespace UE::BuildSelection::Internal;
+	const FSlateColor TitleColor = FStyleColors::AccentWhite;
+	const FSlateFontInfo TitleFont = FCoreStyle::GetDefaultFontStyle("Bold", 10);
+
+	if (FString PastConfiguration; ReadSetting(TEXT("Configuration"), PastConfiguration))
+	{
+		SelectedConfiguration = MakeShared<FString>(MoveTemp(PastConfiguration));
+	}
+
+	this->ChildSlot
+	[
+		SNew(SVerticalBox)
+		+SVerticalBox::Slot()
+		.HAlign(HAlign_Fill)
+		.VAlign(VAlign_Top)
+		.AutoHeight()
+		[
+			SNew(SHorizontalBox)
+			+SHorizontalBox::Slot()
+			.HAlign(HAlign_Left)
+			.VAlign(VAlign_Center)
+			.AutoWidth()
+			[
+				SNew(STextBlock)
+					.ColorAndOpacity(TitleColor)
+					.Font(TitleFont)
+					.Text(LOCTEXT("BuildArtifactSelection_AvailableBuilds", "Available Builds"))
+			]
+			+SHorizontalBox::Slot()
+			.HAlign(HAlign_Fill)
+			.VAlign(VAlign_Center)
+			.Padding(0, 0, 0, 0)
+			[
+				SAssignNew(SelectedGroupArtifactHeaderGrid, SGridPanel)
+			]
+		]
+
+		+ SVerticalBox::Slot()
+		.HAlign(HAlign_Fill)
+		.VAlign(VAlign_Top)
+		.AutoHeight()
+		[
+			SAssignNew(SelectedGroupArtifactGrid, SGridPanel)
+		]
+	];
+}
+
+bool SBuildArtifactSelection::HasArtifactConfigurations() const
+{
+	return !ArtifactConfigurationSet.IsEmpty();
+}
+
+TSharedPtr<FString> SBuildArtifactSelection::GetSelectedConfiguration() const
+{
+	return SelectedConfiguration;
+}
+
+const TArray<FString>& SBuildArtifactSelection::GetSelectedArtifacts() const
+{
+	return SelectedArtifacts;
+}
+
+void SBuildArtifactSelection::SetAvailableConfigurations(
+	TSet<FString>&& InArtifactConfigurationSet,
+	TMultiMap<FString, FString>&& InArtifactNameMap
+)
+{
+	using namespace UE::BuildSelection::Internal;
+	ArtifactConfigurationSet = MoveTemp(InArtifactConfigurationSet);
+	ArtifactNameMap = MoveTemp(InArtifactNameMap);
+
+	ArtifactConfigurationSet.StableSort(TLess());
+	ArtifactNameMap.ValueStableSort(TLess());
+
+	ArtifactConfigurationList.Empty();
+	if (ArtifactConfigurationSet.Num() > 1)
+	{
+		ArtifactConfigurationList.Add(MakeShared<FString>(TEXT("All")));
+	}
+	for (const FString& ArtifactConfiguration : ArtifactConfigurationSet)
+	{
+		ArtifactConfigurationList.Add(MakeShared<FString>(ArtifactConfiguration));
+	}
+	if (!ArtifactConfigurationList.IsEmpty())
+	{
+		ConformSelection(SelectedConfiguration, ArtifactConfigurationList);
+	}
+
+	if (!ArtifactConfigurationList.IsEmpty())
+	{
+		WriteSetting(TEXT("Configuration"), SelectedConfiguration ? **SelectedConfiguration : TEXT(""));
+	}
+}
+
+void SBuildArtifactSelection::Refresh(const TMap<FString, UE::BuildSelection::Internal::FArtifact>* InNamedArtifacts)
+{
+	if (InNamedArtifacts)
+	{
+		NamedArtifacts = *InNamedArtifacts;
+	}
+	else
+	{
+		NamedArtifacts.Empty();
+	}
+	RebuildSelectedGroupArtifactHeader();
+	RebuildSelectedGroupArtifactContents();
+}
+
+void SBuildArtifactSelection::RebuildSelectedGroupArtifactHeader()
+{
+	if (!SelectedGroupArtifactHeaderGrid)
+	{
+		return;
+	}
+
+	SelectedGroupArtifactHeaderGrid->ClearChildren();
+	if (ArtifactConfigurationSet.IsEmpty())
+	{
+		return;
+	}
+
+	SelectedGroupArtifactHeaderGrid->AddSlot(0, 0)
+		.Padding(10,0,2,0)
+		.VAlign(VAlign_Center)
+		[
+			SNew(STextBlock)
+			.Justification(ETextJustify::Right)
+			.Text(LOCTEXT("BuildSelection_Configuration", "Configuration:"))
+		];
+
+	SelectedGroupArtifactHeaderGrid->AddSlot(1, 0)
+		.Padding(0,0,0,0)
+		.VAlign(VAlign_Center)
+		[
+			SNew(SComboBox<TSharedPtr<FString>>)
+			.OptionsSource(&ArtifactConfigurationList)
+			.OnSelectionChanged_Lambda([this](TSharedPtr<FString> NewSelectedConfiguration, ESelectInfo::Type SelectInfo)
+			{
+				if (NewSelectedConfiguration.IsValid() && *NewSelectedConfiguration != *SelectedConfiguration)
+				{
+					SetSelectedConfiguration(NewSelectedConfiguration);
+					RebuildSelectedGroupArtifactContents();
+				}
+			})
+			.OnGenerateWidget_Lambda([this](TSharedPtr<FString> Item)
+				{
+					return SNew(STextBlock)
+						.Text(FText::FromString(*Item));
+				})
+			[
+				SNew(STextBlock)
+					.Text_Lambda([this]()
+					{
+						return FText::FromString(SelectedConfiguration ? **SelectedConfiguration : TEXT(""));
+					})
+			]
+		];
+}
+
+void SBuildArtifactSelection::RebuildSelectedGroupArtifactContents()
+{
+	using namespace UE::BuildSelection::Internal;
+
+	if (!SelectedGroupArtifactGrid)
+	{
+		return;
+	}
+
+	SelectedGroupArtifactGrid->ClearChildren();
+
+	int32 Row = 0;
+	int32 Column = 0;
+	TArray<FString> ArtifactNames;
+	const bool bPresentConsistentArtifactNameSet = false;
+	if (bPresentConsistentArtifactNameSet)
+	{
+		if (!SelectedConfiguration || ArtifactConfigurationSet.IsEmpty())
+		{
+			ArtifactNameMap.GenerateValueArray(ArtifactNames);
+		}
+		else if (*SelectedConfiguration == TEXT("All"))
+		{
+			for (const TPair<FString, FString>& ArtifactNamePair : ArtifactNameMap)
+			{
+				ArtifactNames.AddUnique(ArtifactNamePair.Value);
+			}
+			ArtifactNames.StableSort(TLess());
+		}
+		else
+		{
+			ArtifactNameMap.MultiFind(*SelectedConfiguration, ArtifactNames);
+		}
+	}
+	else
+	{
+		if (!SelectedConfiguration || ArtifactConfigurationSet.IsEmpty() || *SelectedConfiguration == TEXT("All"))
+		{
+			NamedArtifacts.GenerateKeyArray(ArtifactNames);
+		}
+		else
+		{
+			FName SelectedConfigurationName(*SelectedConfiguration);
+			for (const TPair<FString, FArtifact>& NameArtifactPair : NamedArtifacts)
+			{
+				if (NameArtifactPair.Value.Configurations.Contains(SelectedConfigurationName))
+				{
+					ArtifactNames.Add(NameArtifactPair.Key);
+				}
+			}
+		}
+		ArtifactNames.StableSort(TLess());
+	}
+
+	const int32 TotalListCount = ArtifactNames.Num();
+	// Heuristic for choosing grid rows when laying out rows and columns for the quantity of platforms
+	const int32 MaxRows = FMath::Max(4, FMath::FloorToInt32(FMath::Pow(float(TotalListCount), 0.7)));
+
+	auto IsBuildGroupArtifactEnabled = [this] (const FString& ArtifactName)
+	{
+		const FArtifact* Artifact = NamedArtifacts.Find(*ArtifactName);
+		if (!Artifact)
+		{
+			return false;
+		}
+		if (ArtifactConfigurationSet.IsEmpty() || !SelectedConfiguration || *SelectedConfiguration == TEXT("All"))
+		{
+			return true;
+		}
+		if (*SelectedConfiguration == TEXT("None"))
+		{
+			return Artifact->Configurations.IsEmpty();
+		}
+		return Artifact->Configurations.Contains(FName(*SelectedConfiguration));
+	};
+
+	for (const FString& ArtifactName : ArtifactNames)
+	{
+		TSharedPtr<SCheckBox> CurrentCheckbox;
+		SelectedGroupArtifactGrid->AddSlot(Column, Row)
+			.Padding(0,0,0,2)
+			[
+				SNew(SHorizontalBox)
+				.IsEnabled_Lambda([this, ArtifactName, IsBuildGroupArtifactEnabled]
+				{
+					return IsBuildGroupArtifactEnabled(ArtifactName);
+				})
+				+SHorizontalBox::Slot()
+				.AutoWidth()
+				.VAlign(VAlign_Center)
+				[
+					SAssignNew(CurrentCheckbox, SCheckBox)
+					.IsChecked_Lambda([this, ArtifactName, IsBuildGroupArtifactEnabled]
+					{
+						bool bIsChecked = IsBuildGroupArtifactEnabled(ArtifactName) && SelectedArtifacts.Contains(ArtifactName);
+						return bIsChecked ? ECheckBoxState::Checked : ECheckBoxState::Unchecked;
+					})
+					.OnCheckStateChanged_Lambda([this, ArtifactName](ECheckBoxState InNewState)
+					{
+						if (InNewState == ECheckBoxState::Checked)
+						{
+							SelectedArtifacts.AddUnique(ArtifactName);
+						}
+						else if (InNewState == ECheckBoxState::Unchecked)
+						{
+							SelectedArtifacts.Remove(ArtifactName);
+						}
+					})
+				]
+				+SHorizontalBox::Slot()
+				.AutoWidth()
+				.VAlign(VAlign_Center)
+				[
+					SNew(SButton)
+					.ButtonStyle(FAppStyle::Get(), "InvisibleButton")
+					.IsFocusable(false)
+
+					.OnClicked_Lambda([this, CurrentCheckbox, ArtifactName]()
+						{
+							ECheckBoxState NewState = ECheckBoxState::Checked;
+							if (CurrentCheckbox->IsChecked())
+							{
+								NewState = ECheckBoxState::Unchecked;
+							}
+							CurrentCheckbox->SetIsChecked(NewState);
+							if (NewState == ECheckBoxState::Checked)
+							{
+								SelectedArtifacts.AddUnique(ArtifactName);
+							}
+							else if (NewState == ECheckBoxState::Unchecked)
+							{
+								SelectedArtifacts.Remove(ArtifactName);
+							}
+							return FReply::Handled();
+						})
+					[
+						SNew(STextBlock)
+						.Justification(ETextJustify::Left)
+						.Text(FText::FromString(ArtifactName))
+					]
+				]
+				+SHorizontalBox::Slot()
+				.AutoWidth()
+				.VAlign(VAlign_Center)
+				[
+					SNew(SButton)
+					.ButtonStyle(&FAppStyle::Get().GetWidgetStyle<FButtonStyle>("SimpleButton"))
+					.ToolTipText(LOCTEXT("BuildSelection_BrowseToBuildUrl", "Browse to build url"))
+					.IsEnabled_Lambda([this, ArtifactName, IsBuildGroupArtifactEnabled]
+					{
+						if (!IsBuildGroupArtifactEnabled(ArtifactName))
+						{
+							return false;
+						}
+						if (FArtifact* Artifact = NamedArtifacts.Find(ArtifactName))
+						{
+							if (FCbFieldView BuildUrlField = Artifact->BuildRecord.Metadata["buildurl"]; BuildUrlField.HasValue() && !BuildUrlField.HasError())
+							{
+								return true;
+							}
+						}
+						return false;
+					})
+					.OnClicked_Lambda([this, ArtifactName, IsBuildGroupArtifactEnabled]
+					{
+						if (!IsBuildGroupArtifactEnabled(ArtifactName))
+						{
+							return FReply::Unhandled();
+						}
+						if (FArtifact* Artifact = NamedArtifacts.Find(ArtifactName))
+						{
+							if (FCbFieldView BuildUrlField = Artifact->BuildRecord.Metadata["buildurl"]; BuildUrlField.HasValue() && !BuildUrlField.HasError())
+							{
+								FString Url(FUTF8ToTCHAR(BuildUrlField.AsString()));
+								FPlatformProcess::LaunchURL(*Url, nullptr, nullptr);
+								return FReply::Handled();
+							}
+						}
+						return FReply::Unhandled();
+					})
+					[
+						SNew(SImage)
+						.Image(FAppStyle::Get().GetBrush("Zen.BrowserView"))
+						.ColorAndOpacity(FSlateColor::UseForeground())
+					]
+				]
+			];
+
+		if (++Row >= MaxRows)
+		{
+			Column++;
+			Row = 0;
+		}
+	}
+}
+
+void SBuildArtifactSelection::SetSelectedConfiguration(const TSharedPtr<FString> InSelectedConfiguration)
+{
+	using namespace UE::BuildSelection::Internal;
+
+	SelectedConfiguration = InSelectedConfiguration;
+	WriteSetting(TEXT("Configuration"), *SelectedConfiguration);
+}
+
+#undef LOCTEXT_NAMESPACE
diff --git a/Engine/Source/Editor/StorageServerWidgets/Private/SBuildSelection.cpp b/Engine/Source/Editor/StorageServerWidgets/Private/SBuildSelection.cpp
--- a/Engine/Source/Editor/StorageServerWidgets/Private/SBuildSelection.cpp
+++ b/Engine/Source/Editor/StorageServerWidgets/Private/SBuildSelection.cpp
@@ -18,6 +18,7 @@
 #include "Misc/Paths.h"
 #include "Misc/UProjectInfo.h"
 #include "Modules/BuildVersion.h"
+#include "SBuildArtifactSelection.h"
 #include "SBuildTransferFilters.h"
 #include "SMultiSelectComboBox.h"
 #include "Styling/StyleColors.h"
@@ -38,8 +39,6 @@
 
 #define LOCTEXT_NAMESPACE "StorageServerBuild"
 
-const FString SBuildSelection::EmptyString;
-
 namespace UE::BuildSelection::Internal
 {
 	namespace FBuildGroupIds
@@ -168,11 +167,6 @@ void SBuildSelection::Construct(const FArguments& InArgs)
 		bIncludePreflights = FCString::ToBool(*IncludePreflightsString);
 	}
 
-	if (FString PastConfiguration; ReadSetting(TEXT("Configuration"), PastConfiguration))
-	{
-		SelectedConfiguration = MakeShared<FString>(MoveTemp(PastConfiguration));
-	}
-
 	if (TSharedPtr<UE::Zen::Build::FBuildServiceInstance> ServiceInstance = BuildServiceInstance.Get())
 	{
 		ServiceInstance->OnRefreshNamespacesAndBucketsComplete().AddSPLambda(this, [this]()
@@ -498,9 +492,8 @@ void SBuildSelection::RegenerateBuildGroups(UE::Zen::Build::FListBuildsState& Li
 
 	BuildGroups.Empty();
 	PresentedBuildGroups.Reset();
-	ArtifactNameMap.Empty();
-	ArtifactConfigurationSet.Empty();
-	ArtifactConfigurationList.Empty();
+	TMultiMap<FString, FString> ArtifactNameMap;
+	TSet<FString> ArtifactConfigurationSet;
 	bool bHasNoneArtifactConfiguration = false;
 
 	auto MakeGroupKey = [](const FString& Namespace, const FString& CommitIdentifier, const FBuildServiceInstance::FBuildRecord& BuildRecord)
@@ -662,25 +655,7 @@ void SBuildSelection::RegenerateBuildGroups(UE::Zen::Build::FListBuildsState& Li
 		}
 	}
 
-	ArtifactConfigurationSet.StableSort(TLess());
-	ArtifactNameMap.ValueStableSort(TLess());
-	if (!ArtifactConfigurationList.IsEmpty() && bHasNoneArtifactConfiguration)
-	{
-		ArtifactConfigurationList.Add(MakeShared<FString>(TEXT("None")));
-	}
-	if (ArtifactConfigurationSet.Num() > 1)
-	{
-		ArtifactConfigurationList.Add(MakeShared<FString>(TEXT("All")));
-	}
-	for (const FString& ArtifactConfiguration : ArtifactConfigurationSet)
-	{
-		ArtifactConfigurationList.Add(MakeShared<FString>(ArtifactConfiguration));
-	}
-	if (!ArtifactConfigurationList.IsEmpty())
-	{
-		ConformSelection(SelectedConfiguration, ArtifactConfigurationList);
-		WriteSetting(TEXT("Configuration"), SelectedConfiguration ? **SelectedConfiguration : TEXT(""));
-	}
+	BuildArtifactSelection->SetAvailableConfigurations(MoveTemp(ArtifactConfigurationSet), MoveTemp(ArtifactNameMap));
 	KeyedGroups.GenerateValueArray(BuildGroups);
 	RefreshPresentedBuildGroups();
 }
@@ -806,6 +781,8 @@ TSharedRef<SWidget> SBuildSelection::OnGenerateBuildTypeTextBlockFromString(TSha
 
 bool SBuildSelection::DoesPassAllViewFilters(TSharedPtr<FBuildGroup> InItem) const
 {
+	using namespace UE::BuildSelection::Internal;
+
 	if (!InItem)
 	{
 		return false;
@@ -940,259 +917,14 @@ TSharedPtr<SWidget> SBuildSelection::OnGetBuildGroupContextMenuContent() const
 	return MenuBuilder.MakeWidget();
 }
 
-void SBuildSelection::RebuildSelectedGroupArtifactHeader(FBuildSelectionBuildGroupPtr Item)
+void SBuildSelection::RefreshArtifactSelectionWidget(FBuildSelectionBuildGroupPtr Item)
 {
-	if (!SelectedGroupArtifactHeaderGrid)
-	{
-		return;
-	}
-
-	SelectedGroupArtifactHeaderGrid->ClearChildren();
-	if (ArtifactConfigurationSet.IsEmpty())
-	{
-		return;
-	}
-
-	SelectedGroupArtifactHeaderGrid->AddSlot(0, 0)
-		.Padding(10,0,2,0)
-		.VAlign(VAlign_Center)
-		[
-			SNew(STextBlock)
-			.Justification(ETextJustify::Right)
-			.Text(LOCTEXT("BuildSelection_Configuration", "Configuration:"))
-		];
-
-	SelectedGroupArtifactHeaderGrid->AddSlot(1, 0)
-		.Padding(0,0,0,0)
-		.VAlign(VAlign_Center)
-		[
-			SNew(SComboBox<TSharedPtr<FString>>)
-			.OptionsSource(&ArtifactConfigurationList)
-			.OnSelectionChanged_Lambda([this, Item](TSharedPtr<FString> NewSelectedConfiguration, ESelectInfo::Type SelectInfo)
-			{
-				if (NewSelectedConfiguration.IsValid() && *NewSelectedConfiguration != *SelectedConfiguration)
-				{
-					SetSelectedConfiguration(NewSelectedConfiguration);
-					RebuildSelectedGroupArtifactGrid(Item);
-				}
-			})
-			.OnGenerateWidget(this, &SBuildSelection::OnGenerateTextBlockFromString)
-			[
-				SNew(STextBlock)
-					.Text_Lambda([this]()
-					{
-						return FText::FromString(SelectedConfiguration ? **SelectedConfiguration : TEXT(""));
-					})
-			]
-		];
-}
-
-void SBuildSelection::RebuildSelectedGroupArtifactGrid(FBuildSelectionBuildGroupPtr Item)
-{
-	if (!SelectedGroupArtifactGrid)
-	{
-		return;
-	}
-
-	SelectedGroupArtifactGrid->ClearChildren();
-
-	if (!Item)
-	{
-		return;
-	}
-
-	int32 Row = 0;
-	int32 Column = 0;
-	TArray<FString> ArtifactNames;
-	const bool bPresentConsistentArtifactNameSet = false;
-	if (bPresentConsistentArtifactNameSet)
-	{
-		if (!SelectedConfiguration || ArtifactConfigurationSet.IsEmpty())
-		{
-			ArtifactNameMap.GenerateValueArray(ArtifactNames);
-		}
-		else if (*SelectedConfiguration == TEXT("All"))
-		{
-			for (const TPair<FString, FString>& ArtifactNamePair : ArtifactNameMap)
-			{
-				ArtifactNames.AddUnique(ArtifactNamePair.Value);
-			}
-			ArtifactNames.StableSort(TLess());
-		}
-		else
-		{
-			ArtifactNameMap.MultiFind(*SelectedConfiguration, ArtifactNames);
-		}
-	}
-	else
-	{
-		if (!SelectedConfiguration || ArtifactConfigurationSet.IsEmpty() || *SelectedConfiguration == TEXT("All"))
-		{
-			Item->NamedArtifacts.GenerateKeyArray(ArtifactNames);
-		}
-		else
-		{
-			FName SelectedConfigurationName(*SelectedConfiguration);
-			for (const TPair<FString, FArtifact>& NameArtifactPair : Item->NamedArtifacts)
-			{
-				if (NameArtifactPair.Value.Configurations.Contains(SelectedConfigurationName))
-				{
-					ArtifactNames.Add(NameArtifactPair.Key);
-				}
-			}
-		}
-		ArtifactNames.StableSort(TLess());
-	}
-
-	const int32 TotalListCount = ArtifactNames.Num();
-	// Heuristic for choosing grid rows when laying out rows and columns for the quantity of platforms
-	const int32 MaxRows = FMath::Max(4, FMath::FloorToInt32(FMath::Pow(float(TotalListCount), 0.7)));
-
-	auto IsBuildGroupArtifactEnabled = [this, Item] (const FString& ArtifactName)
-	{
-		const FArtifact* Artifact = Item->NamedArtifacts.Find(*ArtifactName);
-		if (!Artifact)
-		{
-			return false;
-		}
-		if (ArtifactConfigurationSet.IsEmpty() || !SelectedConfiguration || *SelectedConfiguration == TEXT("All"))
-		{
-			return true;
-		}
-		if (*SelectedConfiguration == TEXT("None"))
-		{
-			return Artifact->Configurations.IsEmpty();
-		}
-		return Artifact->Configurations.Contains(FName(*SelectedConfiguration));
-	};
-
-	for (const FString& ArtifactName : ArtifactNames)
-	{
-		TSharedPtr<SCheckBox> CurrentCheckbox;
-		SelectedGroupArtifactGrid->AddSlot(Column, Row)
-			.Padding(0,0,0,2)
-			[
-				SNew(SHorizontalBox)
-				.IsEnabled_Lambda([this, Item, ArtifactName, IsBuildGroupArtifactEnabled]
-				{
-					return IsBuildGroupArtifactEnabled(ArtifactName);
-				})
-				+SHorizontalBox::Slot()
-				.AutoWidth()
-				.VAlign(VAlign_Center)
-				[
-					SAssignNew(CurrentCheckbox, SCheckBox)
-					.IsChecked_Lambda([this, Item, ArtifactName, IsBuildGroupArtifactEnabled]
-					{
-						bool bIsChecked = IsBuildGroupArtifactEnabled(ArtifactName) && SelectedGroupSelectedArtifacts.Contains(ArtifactName);
-						return bIsChecked ? ECheckBoxState::Checked : ECheckBoxState::Unchecked;
-					})
-					.OnCheckStateChanged_Lambda([this, ArtifactName](ECheckBoxState InNewState)
-					{
-						if (InNewState == ECheckBoxState::Checked)
-						{
-							SelectedGroupSelectedArtifacts.AddUnique(ArtifactName);
-						}
-						else if (InNewState == ECheckBoxState::Unchecked)
-						{
-							SelectedGroupSelectedArtifacts.Remove(ArtifactName);
-						}
-					})
-				]
-				+SHorizontalBox::Slot()
-				.AutoWidth()
-				.VAlign(VAlign_Center)
-				[
-					SNew(SButton)
-					.ButtonStyle(FAppStyle::Get(), "InvisibleButton")
-					.IsFocusable(false)
-
-					.OnClicked_Lambda([this, CurrentCheckbox, ArtifactName]()
-						{
-							ECheckBoxState NewState = ECheckBoxState::Checked;
-							if (CurrentCheckbox->IsChecked())
-							{
-								NewState = ECheckBoxState::Unchecked;
-							}
-							CurrentCheckbox->SetIsChecked(NewState);
-							if (NewState == ECheckBoxState::Checked)
-							{
-								SelectedGroupSelectedArtifacts.AddUnique(ArtifactName);
-							}
-							else if (NewState == ECheckBoxState::Unchecked)
-							{
-								SelectedGroupSelectedArtifacts.Remove(ArtifactName);
-							}
-							return FReply::Handled();
-						})
-					[
-						SNew(STextBlock)
-						.Justification(ETextJustify::Left)
-						.Text(FText::FromString(ArtifactName))
-					]
-				]
-				+SHorizontalBox::Slot()
-				.AutoWidth()
-				.VAlign(VAlign_Center)
-				[
-					SNew(SButton)
-					.ButtonStyle(&FAppStyle::Get().GetWidgetStyle<FButtonStyle>("SimpleButton"))
-					.ToolTipText(LOCTEXT("BuildSelection_BrowseToBuildUrl", "Browse to build url"))
-					.IsEnabled_Lambda([this, Item, ArtifactName, IsBuildGroupArtifactEnabled]
-					{
-						if (!IsBuildGroupArtifactEnabled(ArtifactName))
-						{
-							return false;
-						}
-						if (FArtifact* Artifact = Item->NamedArtifacts.Find(ArtifactName))
-						{
-							if (FCbFieldView BuildUrlField = Artifact->BuildRecord.Metadata["buildurl"]; BuildUrlField.HasValue() && !BuildUrlField.HasError())
-							{
-								return true;
-							}
-						}
-						return false;
-					})
-					.OnClicked_Lambda([this, Item, ArtifactName, IsBuildGroupArtifactEnabled]
-					{
-						if (!IsBuildGroupArtifactEnabled(ArtifactName))
-						{
-							return FReply::Unhandled();
-						}
-						if (FArtifact* Artifact = Item->NamedArtifacts.Find(ArtifactName))
-						{
-							if (FCbFieldView BuildUrlField = Artifact->BuildRecord.Metadata["buildurl"]; BuildUrlField.HasValue() && !BuildUrlField.HasError())
-							{
-								FString Url(FUTF8ToTCHAR(BuildUrlField.AsString()));
-								FPlatformProcess::LaunchURL(*Url, nullptr, nullptr);
-								return FReply::Handled();
-							}
-						}
-						return FReply::Unhandled();
-					})
-					[
-						SNew(SImage)
-						.Image(FAppStyle::Get().GetBrush("Zen.BrowserView"))
-						.ColorAndOpacity(FSlateColor::UseForeground())
-					]
-				]
-			];
-
-		if (++Row >= MaxRows)
-		{
-			Column++;
-			Row = 0;
-		}
-	}
+	BuildArtifactSelection->Refresh(Item ? &Item->NamedArtifacts : nullptr);
 }
 
 void SBuildSelection::BuildGroupSelectionChanged(FBuildSelectionBuildGroupPtr Item, ESelectInfo::Type SelectInfo)
 {
-	using namespace UE::Zen::Build;
-
-	RebuildSelectedGroupArtifactHeader(Item);
-
-	RebuildSelectedGroupArtifactGrid(Item);
+	RefreshArtifactSelectionWidget(Item);
 }
 
 EColumnSortMode::Type SBuildSelection::GetColumnSortMode(const FName ColumnId) const
@@ -1253,6 +985,7 @@ int SBuildSelection::ConvertBuildTypeToDefaultSelectionOrdering(const FString& I
 
 void SBuildSelection::SetUserSelectedDestination(const FStringView InDestination)
 {
+	using namespace UE::BuildSelection::Internal;
 	EBuildType BuildType = GetSelectedBuildType();
 	int BuildTypeIndex = (int)BuildType;
 	UserSelectedDestinations[BuildTypeIndex] = InDestination;
@@ -1275,6 +1008,7 @@ void SBuildSelection::SetUserSelectedDestination(const FStringView InDestination
 
 void SBuildSelection::SetUserSelectedAppendBuildNameToDestination(bool bAppendBuildNameToDestination)
 {
+	using namespace UE::BuildSelection::Internal;
 	EBuildType BuildType = GetSelectedBuildType();
 	int BuildTypeIndex = (int)BuildType;
 	bAppendBuildNameToDestinations[BuildTypeIndex] = bAppendBuildNameToDestination;
@@ -1284,6 +1018,7 @@ void SBuildSelection::SetUserSelectedAppendBuildNameToDestination(bool bAppendBu
 
 void SBuildSelection::SetUserSelectedTransferFilter(const FStringView InIncludeFilter, const FStringView InExcludeFilter)
 {
+	using namespace UE::BuildSelection::Internal;
 	EBuildType BuildType = GetSelectedBuildType();
 	int BuildTypeIndex = (int)BuildType;
 	UserSelectedTransferFilters[BuildTypeIndex].IncludeFilter = InIncludeFilter;
@@ -1380,62 +1115,33 @@ FString SBuildSelection::GetEffectiveDestination() const
 
 void SBuildSelection::SetSelectedStream(const TSharedPtr<FString> InSelectedStream)
 {
+	using namespace UE::BuildSelection::Internal;
 	SelectedStream = InSelectedStream;
 	WriteSetting(TEXT("Branch"), *SelectedStream);
 }
 
 void SBuildSelection::SetSelectedProject(const TSharedPtr<FString> InSelectedProject)
 {
+	using namespace UE::BuildSelection::Internal;
 	SelectedProject = InSelectedProject;
 	WriteSetting(TEXT("Project"), *SelectedProject);
 }
 
 void SBuildSelection::SetSelectedBuildType(const TSharedPtr<FString> InBuildType)
 {
+	using namespace UE::BuildSelection::Internal;
 	SelectedBuildType = InBuildType;
 	WriteSetting(TEXT("BuildType"), *SelectedBuildType);
 }
 
 void SBuildSelection::SetIncludePreflights(const bool bInIncludePreflights)
 {
+	using namespace UE::BuildSelection::Internal;
 	bIncludePreflights = bInIncludePreflights;
 	WriteSetting(TEXT("IncludePreflights"), bIncludePreflights ? TEXT("True") : TEXT("False"));
 	PostViewFilterChange();
 }
 
-void SBuildSelection::SetSelectedConfiguration(const TSharedPtr<FString> InSelectedConfiguration)
-{
-	SelectedConfiguration = InSelectedConfiguration;
-	WriteSetting(TEXT("Configuration"), *SelectedConfiguration);
-}
-
-bool SBuildSelection::WriteSetting(const FStringView InSectionName, const FStringView InKeyName, const FStringView InValue)
-{
-	static const FString SectionName = TEXT("Unreal Engine/Build Storage");
-	const FString FinalSectionName = InSectionName.IsEmpty() ? SectionName : FPaths::Combine(SectionName, InSectionName);
-	const FString FinalKeyName(InKeyName);
-	FString Value(InValue);
-	return FPlatformMisc::SetStoredValue(TEXT("Epic Games"), FinalSectionName, FinalKeyName, Value);
-}
-
-bool SBuildSelection::WriteSetting(const FStringView InKeyName, const FStringView InValue)
-{
-	return WriteSetting(EmptyString, InKeyName, InValue);
-}
-
-bool SBuildSelection::ReadSetting(const FStringView InSectionName, const FStringView InKeyName, FString& OutValue)
-{
-	static const FString SectionName = TEXT("Unreal Engine/Build Storage");
-	const FString FinalSectionName = InSectionName.IsEmpty() ? SectionName : FPaths::Combine(SectionName, InSectionName);
-	const FString FinalKeyName(InKeyName);
-	return FPlatformMisc::GetStoredValue(TEXT("Epic Games"), FinalSectionName, FinalKeyName, OutValue);
-}
-
-bool SBuildSelection::ReadSetting(const FStringView InKeyName, FString& OutValue)
-{
-	return ReadSetting(EmptyString, InKeyName, OutValue);
-}
-
 void SBuildSelection::SetUserSelectedProjectDictionaryRoot(const FStringView InRoot)
 {
 	if (InRoot.IsEmpty())
@@ -1544,22 +1250,6 @@ FString SBuildSelection::SanitizeBucket(const FString& InString)
 	return OutString;
 }
 
-void SBuildSelection::ConformSelection(TSharedPtr<FString>& SelectedItem, const TArray<TSharedPtr<FString>>& SelectionList)
-{
-	if (!SelectedItem)
-	{
-		SelectedItem = SelectionList.IsEmpty() ? nullptr : SelectionList[0];
-		return;
-	}
-
-	const TSharedPtr<FString>* FoundSelectionListItem = SelectionList.FindByPredicate([&SelectedItem](const TSharedPtr<FString>& Item)
-		{
-			return *Item == *SelectedItem;
-		});
-
-	SelectedItem = FoundSelectionListItem ? *FoundSelectionListItem : SelectionList.IsEmpty() ? nullptr : SelectionList[0];
-}
-
 const TCHAR* SBuildSelection::LexToString(EBuildType BuildType)
 {
 	static const TCHAR* Strings[] =
@@ -1602,6 +1292,7 @@ FReply SBuildSelection::ExploreDestination_OnClicked()
 void SBuildSelection::DownloadSelectedArtifactsInSelectedBuildGroup()
 {
 	using namespace UE::Zen::Build;
+	using namespace UE::BuildSelection::Internal;
 
 	if (TSharedPtr<FBuildServiceInstance> ServiceInstance = BuildServiceInstance.Get())
 	{
@@ -1611,12 +1302,13 @@ void SBuildSelection::DownloadSelectedArtifactsInSelectedBuildGroup()
 			return;
 		}
 
+		const TSharedPtr<FString> SelectedConfiguration = BuildArtifactSelection->GetSelectedConfiguration();
 		FName SelectedConfigurationName(NAME_None);
-		if (!ArtifactConfigurationSet.IsEmpty() && SelectedConfiguration && *SelectedConfiguration != TEXT("All"))
+		if (BuildArtifactSelection->HasArtifactConfigurations() && SelectedConfiguration && *SelectedConfiguration != TEXT("All"))
 		{
 			SelectedConfigurationName = FName(*SelectedConfiguration);
 		}
-		for (const FString& ArtifactForSelectedGroup : SelectedGroupSelectedArtifacts)
+		for (const FString& ArtifactForSelectedGroup : BuildArtifactSelection->GetSelectedArtifacts())
 		{
 			if (FArtifact* Artifact = SelectedItems[0]->NamedArtifacts.Find(ArtifactForSelectedGroup))
 			{
@@ -2391,38 +2083,7 @@ TSharedRef<SWidget> SBuildSelection::GetGridPanel()
 			+ SScrollBox::Slot()
 			.FillSize(1.0f)
 			[
-				SNew(SVerticalBox)
-				+SVerticalBox::Slot()
-				.HAlign(HAlign_Fill)
-				.VAlign(VAlign_Top)
-				.AutoHeight()
-				[
-					SNew(SHorizontalBox)
-					+SHorizontalBox::Slot()
-					.HAlign(HAlign_Left)
-					.VAlign(VAlign_Center)
-					.AutoWidth()
-					[
-						SNew(STextBlock)
-							.ColorAndOpacity(TitleColor)
-							.Font(TitleFont)
-							.Text(LOCTEXT("BuildSelection_AvailableBuilds", "Available Builds"))
-					]
-					+SHorizontalBox::Slot()
-					.HAlign(HAlign_Fill)
-					.VAlign(VAlign_Center)
-					.Padding(0, 0, 0, 0)
-					[
-						SAssignNew(SelectedGroupArtifactHeaderGrid, SGridPanel)
-					]
-				]
-				+SVerticalBox::Slot()
-				.HAlign(HAlign_Fill)
-				.VAlign(VAlign_Top)
-				.AutoHeight()
-				[
-					SAssignNew(SelectedGroupArtifactGrid, SGridPanel)
-				]
+				SAssignNew(BuildArtifactSelection, SBuildArtifactSelection)
 			]
 		]
 		+SHorizontalBox::Slot()
@@ -2493,12 +2154,13 @@ TSharedRef<SWidget> SBuildSelection::GetGridPanel()
 						return false;
 					}
 
+					const TSharedPtr<FString> SelectedConfiguration = BuildArtifactSelection->GetSelectedConfiguration();
 					FName SelectedConfigurationName(NAME_None);
-					if (!ArtifactConfigurationSet.IsEmpty() && SelectedConfiguration && *SelectedConfiguration != TEXT("All"))
+					if (BuildArtifactSelection->HasArtifactConfigurations() && SelectedConfiguration && *SelectedConfiguration != TEXT("All"))
 					{
 						SelectedConfigurationName = FName(*SelectedConfiguration);
 					}
-					for (const FString& ArtifactForSelectedGroup : SelectedGroupSelectedArtifacts)
+					for (const FString& ArtifactForSelectedGroup : BuildArtifactSelection->GetSelectedArtifacts())
 					{
 						if (FArtifact* FoundArtifact = SelectedItems[0]->NamedArtifacts.Find(ArtifactForSelectedGroup))
 						{
