Allow BuildStorageTool to accept a WorkspaceRoot commandline argument that will ensure that we use that as the user selected engine installation even if there is no registered engine installation at that path.

#Synced-CL 51289657

diff --git a/Engine/Source/Editor/StorageServerWidgets/Private/SBuildSelection.cpp b/Engine/Source/Editor/StorageServerWidgets/Private/SBuildSelection.cpp
--- a/Engine/Source/Editor/StorageServerWidgets/Private/SBuildSelection.cpp
+++ b/Engine/Source/Editor/StorageServerWidgets/Private/SBuildSelection.cpp
@@ -74,6 +74,35 @@ void SBuildSelection::Construct(const FArguments& InArgs)
 		DesktopPlatform->EnumerateEngineInstallations(EngineInstallations);
 	}
 
+	// If the tool was launched with an explicit workspace root directory, use it to set the engine installation directory
+	// and ensure that it is considered as an installation directory even if is not one of the items found by EnumerateEngineInstallations.
+	// This can happen when an engine installation hasn't had the editor run in it yet, and hasn't yet registered itself in a way that is
+	// discoverable by EnumerateEngineInstallations.
+	FString WorkspaceRoot;
+	if (FParse::Value(FCommandLine::Get(), TEXT("-WorkspaceRoot"), WorkspaceRoot))
+	{
+		WorkspaceRoot.TrimStartAndEndInline();
+		WorkspaceRoot.TrimQuotesInline();
+		FPaths::NormalizeDirectoryName(WorkspaceRoot);
+		bool bFoundExisting = false;
+		for (const TPair<FString, FString>& EngineInstallation : EngineInstallations)
+		{
+			if (FPaths::IsSamePath(WorkspaceRoot, EngineInstallation.Value))
+			{
+				bFoundExisting = true;
+				WriteSetting(TEXT("UserSelectedEngineInstallation"), EngineInstallation.Key);
+				SetUserSelectedProjectDictionaryRoot(EngineInstallation.Value);
+				break;
+			}
+		}
+		if (!bFoundExisting)
+		{
+			EngineInstallations.Add(FString(), WorkspaceRoot);
+			WriteSetting(TEXT("UserSelectedEngineInstallation"), FString());
+			SetUserSelectedProjectDictionaryRoot(WorkspaceRoot);
+		}
+	}
+
 	// If the tool was launched with an explicit project file, use it to set the engine installation directory and the project
 	// regardless of past settings.  These values should then persist as new settings.
 	FString ProjectFile;
