Adding URL highlighting and downloading support for BuildStorageTool.

#Synced-CL 51529686

diff --git a/Engine/Source/Developer/Zen/Internal/ZenBuildUtils.h b/Engine/Source/Developer/Zen/Internal/ZenBuildUtils.h
--- a/Engine/Source/Developer/Zen/Internal/ZenBuildUtils.h
+++ b/Engine/Source/Developer/Zen/Internal/ZenBuildUtils.h
@@ -4,18 +4,32 @@
 
 #include "Experimental/BuildServerInterface.h"
 
+#define UE_API ZEN_API
+
 namespace UE::Zen::Build
 {
-	struct FBuildState
-	{
-		FString Namespace;
-		TArray<UE::Zen::Build::FBuildServiceInstance::FBuildRecord> Results;
-	};
-
-	struct FListBuildsState
-	{
-		std::atomic<uint32> PendingQueries;
-		TArray<FBuildState> QueryState;
-	};
+
+struct FBuildState
+{
+	FString Namespace;
+	TArray<UE::Zen::Build::FBuildServiceInstance::FBuildRecord> Results;
+};
+
+struct FListBuildsState
+{
+	std::atomic<uint32> PendingQueries;
+	TArray<FBuildState> QueryState;
+};
+
+struct FBuildReference
+{
+	FString Namespace;
+	FString Bucket;
+	FCbObjectId BuildId;
+};
+
+UE_API bool TryParseBuildReferenceFromUrl(FStringView Url, FBuildReference& OutReference);
+
 } // namespace UE::Zen::Build
 
+#undef UE_API
diff --git a/Engine/Source/Developer/Zen/Private/BuildServerInterface.cpp b/Engine/Source/Developer/Zen/Private/BuildServerInterface.cpp
--- a/Engine/Source/Developer/Zen/Private/BuildServerInterface.cpp
+++ b/Engine/Source/Developer/Zen/Private/BuildServerInterface.cpp
@@ -484,8 +484,8 @@ FBuildServiceInstance::ListBuildsAcrossBuckets(FStringView Namespace, FStringVie
 		Writer.AddString("bucketRegex", BucketRegex);
 		Writer.AddObject("query", Query);
 		Writer.BeginObject("options");
-		Writer.AddInteger("limit", 5000);
-		Writer.AddInteger("max", 50000);
+		Writer.AddInteger("limit", 50000);
+		Writer.AddInteger("max", 500000);
 		Writer.EndObject();
 		Writer.EndObject();
 		Writer.Save(*QueryFileArchive);
@@ -864,7 +864,10 @@ FBuildServiceInstance::FBuildTransfer::RequestCancel()
 	}
 	State->bCancelRequested.store(true);
 	FWriteScopeLock _(State->Lock);
-	State->TransferEndTime = FMonotonicTimePoint::Now();
+	if (State->Status == EBuildTransferStatus::Active)
+	{
+		State->TransferEndTime = FMonotonicTimePoint::Now();
+	}
 	State->Status = EBuildTransferStatus::Canceled;
 }
 
@@ -1164,7 +1167,10 @@ FBuildServiceInstance::BuildTransferThreadLoop()
 		if (State->bCancelRequested.load(std::memory_order_relaxed))
 		{
 			FWriteScopeLock _(State->Lock);
-			State->TransferEndTime = FMonotonicTimePoint::Now();
+			if (State->Status == EBuildTransferStatus::Active)
+			{
+				State->TransferEndTime = FMonotonicTimePoint::Now();
+			}
 			State->Status = EBuildTransferStatus::Canceled;
 			continue;
 		}
diff --git a/Engine/Source/Developer/Zen/Private/ZenBuildUtils.cpp b/Engine/Source/Developer/Zen/Private/ZenBuildUtils.cpp
new file mode 100644
--- /dev/null
+++ b/Engine/Source/Developer/Zen/Private/ZenBuildUtils.cpp
@@ -0,0 +1,44 @@
+// Copyright Epic Games, Inc. All Rights Reserved.
+#include "ZenBuildUtils.h"
+#include "GenericPlatform/GenericPlatformHttp.h"
+#include "String/ParseTokens.h"
+
+namespace UE::Zen::Build
+{
+
+bool TryParseBuildReferenceFromUrl(FStringView Url, FBuildReference& OutReference)
+{
+	// Permissively allow any url with a path that ends with /<namespace>/<bucket>/<buildid>
+	FString Path = FGenericPlatformHttp::GetUrlPath(Url);
+	if (Path.IsEmpty())
+	{
+		return false;
+	}
+
+	TArray<FStringView> PathParts;
+	UE::String::ParseTokens(Path,
+		TCHAR('/'),
+		[&PathParts](FStringView PathPart)
+		{
+			PathParts.Add(PathPart);
+		}, UE::String::EParseTokensOptions::SkipEmpty);
+
+	if (PathParts.Num() < 3)
+	{
+		return false;
+	}
+
+	FCbObjectId::ByteArray BuildIdBytes;
+	FString BuildIdString(PathParts[PathParts.Num() - 1]);
+	if (UE::String::HexToBytes(BuildIdString, BuildIdBytes) != sizeof(BuildIdBytes))
+	{
+		return false;
+	}
+
+	OutReference.Namespace = PathParts[PathParts.Num() - 3];
+	OutReference.Bucket = PathParts[PathParts.Num() - 2];
+	OutReference.BuildId = FCbObjectId(BuildIdBytes);
+	return true;
+}
+
+} // namespace UE::Zen::Build
diff --git a/Engine/Source/Editor/StorageServerWidgets/Internal/SBuildArtifactSelection.h b/Engine/Source/Editor/StorageServerWidgets/Internal/SBuildArtifactSelection.h
--- a/Engine/Source/Editor/StorageServerWidgets/Internal/SBuildArtifactSelection.h
+++ b/Engine/Source/Editor/StorageServerWidgets/Internal/SBuildArtifactSelection.h
@@ -41,7 +41,7 @@ public:
 			TSet<FString>&& InArtifactConfigurationSet,
 			TMultiMap<FString, FString>&& InArtifactNameMap
 		);
-	UE_API void Refresh(const TMap<FString, UE::BuildSelection::Internal::FArtifact>* InNamedArtifacts);
+	UE_API void Refresh(const TMap<FString, UE::BuildSelection::Internal::FArtifact>* InNamedArtifacts, FCbObjectId InHighlightBuildId = {});
 private:
 	void RebuildSelectedGroupArtifactHeader();
 	void RebuildSelectedGroupArtifactContents();
@@ -52,6 +52,7 @@ private:
 	TArray<TSharedPtr<FString>> ArtifactConfigurationList;
 	TSharedPtr<FString> SelectedConfiguration;
 	TArray<FString> SelectedArtifacts;
+	FCbObjectId HighlightedBuildId;
 
 	TMap<FString, UE::BuildSelection::Internal::FArtifact> NamedArtifacts;
 
diff --git a/Engine/Source/Editor/StorageServerWidgets/Internal/SBuildSelection.h b/Engine/Source/Editor/StorageServerWidgets/Internal/SBuildSelection.h
--- a/Engine/Source/Editor/StorageServerWidgets/Internal/SBuildSelection.h
+++ b/Engine/Source/Editor/StorageServerWidgets/Internal/SBuildSelection.h
@@ -17,7 +17,6 @@
 #include "Widgets/Views/SListView.h"
 #include "Widgets/Views/STableRow.h"
 #include "Widgets/Views/STableViewBase.h"
-#include "ZenBuildUtils.h"
 
 #define UE_API STORAGESERVERWIDGETS_API
 
@@ -28,6 +27,7 @@ class SMultiSelectComboBox;
 namespace UE::Zen::Build
 { 
 	struct FListBuildsState;
+	struct FBuildReference;
 }
 
 DECLARE_DELEGATE_TwoParams(FPreBuildTransfer, FString& /*HostOverride*/, UE::Zen::Build::EBuildTransferRequestFlags& /*RequestFlags*/);
@@ -67,6 +67,13 @@ class SBuildSelection : public SCompoundWidget
 
 	UE_API void Construct(const FArguments& InArgs);
 
+	enum class EBuildArtifactAction
+	{
+		None,
+		Highlight,
+		HighlightAndDownload
+	};
+	UE_API void ActOnBuildArtifact(EBuildArtifactAction Action, const UE::Zen::Build::FBuildReference& BuildReference, bool bRebuildLists = true);
 private:
 	enum class EBuildType
 	{
@@ -106,6 +113,7 @@ private:
 	void BuildGroupSelectionChanged(TSharedPtr<FBuildGroup> Item, ESelectInfo::Type SelectInfo);
 	EColumnSortMode::Type GetColumnSortMode(const FName ColumnId) const;
 	void OnColumnSortModeChanged(const EColumnSortPriority::Type SortPriority, const FName& ColumnId, const EColumnSortMode::Type InSortMode);
+	void EnforceTargetedBuildActionOnGroup();
 
 	FText ConvertBuildTypeToText(const FString& InBuildType);
 	int ConvertBuildTypeToDefaultSelectionOrdering(const FString& InBuildType);
@@ -135,6 +143,7 @@ private:
 
 	FReply ExploreDestination_OnClicked();
 	void DownloadSelectedArtifactsInSelectedBuildGroup();
+	bool DownloadItem(const FBuildGroup& Group, const FString& ArtifactName, const UE::BuildSelection::Internal::FArtifact& Artifact);
 
 	TMap<FString, FString> EngineInstallations;
 
@@ -144,9 +153,6 @@ private:
 	FString UserSelectedProjectDictionaryRootDir;
 	TUniquePtr<FUProjectDictionary> UserSelectedProjectDictionary;
 
-	FString BuildUrl;
-	FString BuildId;
-	FString BuildVersion;
 	FString BranchName;
 	int EffectiveCompatibleChangelist = 0;
 
@@ -168,6 +174,9 @@ private:
 		FString ExcludeFilter;
 	};
 
+	std::atomic<EBuildArtifactAction> TargetedBuildIdAction = EBuildArtifactAction::None;
+	std::atomic<FCbObjectId> TargetedBuildId;
+
 	TArray<FString> ActivePlatformFilters;
 	FString UserSelectedDestinations[(int)EBuildType::Count];
 	FTransferFilter UserSelectedTransferFilters[(int)EBuildType::Count];
diff --git a/Engine/Source/Editor/StorageServerWidgets/Private/SBuildArtifactSelection.cpp b/Engine/Source/Editor/StorageServerWidgets/Private/SBuildArtifactSelection.cpp
--- a/Engine/Source/Editor/StorageServerWidgets/Private/SBuildArtifactSelection.cpp
+++ b/Engine/Source/Editor/StorageServerWidgets/Private/SBuildArtifactSelection.cpp
@@ -118,8 +118,11 @@ void SBuildArtifactSelection::SetAvailableConfigurations(
 	}
 }
 
-void SBuildArtifactSelection::Refresh(const TMap<FString, UE::BuildSelection::Internal::FArtifact>* InNamedArtifacts)
+void SBuildArtifactSelection::Refresh(const TMap<FString, UE::BuildSelection::Internal::FArtifact>* InNamedArtifacts, FCbObjectId InHighlightBuildId)
 {
+	using namespace UE::BuildSelection::Internal;
+
+	HighlightedBuildId = InHighlightBuildId;
 	if (InNamedArtifacts)
 	{
 		NamedArtifacts = *InNamedArtifacts;
@@ -128,6 +131,15 @@ void SBuildArtifactSelection::Refresh(const TMap<FString, UE::BuildSelection::In
 	{
 		NamedArtifacts.Empty();
 	}
+
+	for (const TPair<FString, FArtifact>& NameArtifactPair : NamedArtifacts)
+	{
+		if (NameArtifactPair.Value.BuildRecord.BuildId == HighlightedBuildId)
+		{
+			SelectedArtifacts.AddUnique(NameArtifactPair.Key);
+		}
+	}
+
 	RebuildSelectedGroupArtifactHeader();
 	RebuildSelectedGroupArtifactContents();
 }
@@ -228,7 +240,8 @@ void SBuildArtifactSelection::RebuildSelectedGroupArtifactContents()
 			FName SelectedConfigurationName(*SelectedConfiguration);
 			for (const TPair<FString, FArtifact>& NameArtifactPair : NamedArtifacts)
 			{
-				if (NameArtifactPair.Value.Configurations.Contains(SelectedConfigurationName))
+				if ((NameArtifactPair.Value.BuildRecord.BuildId == HighlightedBuildId) ||
+					NameArtifactPair.Value.Configurations.Contains(SelectedConfigurationName))
 				{
 					ArtifactNames.Add(NameArtifactPair.Key);
 				}
@@ -248,6 +261,10 @@ void SBuildArtifactSelection::RebuildSelectedGroupArtifactContents()
 		{
 			return false;
 		}
+		if (Artifact->BuildRecord.BuildId == HighlightedBuildId)
+		{
+			return true;
+		}
 		if (ArtifactConfigurationSet.IsEmpty() || !SelectedConfiguration || *SelectedConfiguration == TEXT("All"))
 		{
 			return true;
@@ -258,6 +275,8 @@ void SBuildArtifactSelection::RebuildSelectedGroupArtifactContents()
 		}
 		return Artifact->Configurations.Contains(FName(*SelectedConfiguration));
 	};
+	const FSlateFontInfo RegularFont = FCoreStyle::GetDefaultFontStyle("Regular", 10);
+	const FSlateFontInfo HighlightFont = FCoreStyle::GetDefaultFontStyle("Bold", 10);
 
 	for (const FString& ArtifactName : ArtifactNames)
 	{
@@ -322,6 +341,17 @@ void SBuildArtifactSelection::RebuildSelectedGroupArtifactContents()
 						SNew(STextBlock)
 						.Justification(ETextJustify::Left)
 						.Text(FText::FromString(ArtifactName))
+						.Font_Lambda([this, ArtifactName, RegularFont, HighlightFont]
+						{
+							if (FArtifact* Artifact = NamedArtifacts.Find(ArtifactName))
+							{
+								if (Artifact->BuildRecord.BuildId == HighlightedBuildId)
+								{
+									return HighlightFont;
+								}
+							}
+							return RegularFont;
+						})
 					]
 				]
 				+SHorizontalBox::Slot()
diff --git a/Engine/Source/Editor/StorageServerWidgets/Private/SBuildSelection.cpp b/Engine/Source/Editor/StorageServerWidgets/Private/SBuildSelection.cpp
--- a/Engine/Source/Editor/StorageServerWidgets/Private/SBuildSelection.cpp
+++ b/Engine/Source/Editor/StorageServerWidgets/Private/SBuildSelection.cpp
@@ -21,6 +21,7 @@
 #include "SBuildArtifactSelection.h"
 #include "SBuildTransferFilters.h"
 #include "SMultiSelectComboBox.h"
+#include "String/ParseTokens.h"
 #include "Styling/StyleColors.h"
 #include "Widgets/Images/SImage.h"
 #include "Widgets/Input/SButton.h"
@@ -34,6 +35,7 @@
 #include "Widgets/Text/STextBlock.h"
 #include "Widgets/Views/SHeaderRow.h"
 #include "Widgets/Views/SListView.h"
+#include "ZenBuildUtils.h"
 #include "ZenServiceInstanceManager.h"
 #include <atomic>
 
@@ -170,9 +172,6 @@ void SBuildSelection::Construct(const FArguments& InArgs)
 	FBuildVersion CurrentBuildVersion;
 	if (FBuildVersion::TryRead(FBuildVersion::GetDefaultFileName(), CurrentBuildVersion))
 	{
-		BuildUrl = CurrentBuildVersion.BuildUrl;
-		BuildId = CurrentBuildVersion.BuildId;
-		BuildVersion = CurrentBuildVersion.BuildVersion;
 		BranchName = SanitizeBranch(CurrentBuildVersion.BranchName);
 
 		EffectiveCompatibleChangelist = CurrentBuildVersion.GetEffectiveCompatibleChangelist();
@@ -222,6 +221,95 @@ void SBuildSelection::Construct(const FArguments& InArgs)
 	];
 }
 
+void SBuildSelection::ActOnBuildArtifact(EBuildArtifactAction Action, const UE::Zen::Build::FBuildReference& BuildReference, bool bAllowRebuildLists)
+{
+	if (Action == EBuildArtifactAction::None)
+	{
+		return;
+	}
+
+	TArray<FStringView> BucketParts;
+	UE::String::ParseTokens(BuildReference.Bucket,
+		TCHAR('.'),
+		[&BucketParts](FStringView BucketPart)
+		{
+			BucketParts.Add(BucketPart);
+		});
+
+	if (BucketParts.Num() < 3)
+	{
+		return;
+	}
+
+	bool bNeedListRebuild = false;
+	if (!SelectedStream || *SelectedStream != BucketParts[2])
+	{
+		SetSelectedStream(MakeShared<FString>(BucketParts[2]));
+		bNeedListRebuild = true;
+	}
+	if (!SelectedProject || *SelectedProject != BucketParts[0])
+	{
+		SetSelectedProject(MakeShared<FString>(BucketParts[0]));
+		bNeedListRebuild = true;
+	}
+	if (!SelectedBuildType || *SelectedBuildType != BucketParts[1])
+	{
+		SetSelectedBuildType(MakeShared<FString>(BucketParts[1]));
+		bNeedListRebuild = true;
+	}
+
+	TargetedBuildIdAction.store(Action);
+	TargetedBuildId.store(BuildReference.BuildId);
+
+	if (bAllowRebuildLists && bNeedListRebuild)
+	{
+		ExecuteOnGameThread(UE_SOURCE_LOCATION,
+			[this]
+			{
+				RebuildLists();
+			});
+	}
+	else
+	{
+		ExecuteOnGameThread(UE_SOURCE_LOCATION,
+			[this]
+			{
+				EnforceTargetedBuildActionOnGroup();
+			});
+	}
+}
+
+void SBuildSelection::EnforceTargetedBuildActionOnGroup()
+{
+	using namespace UE::BuildSelection::Internal;
+
+	EBuildArtifactAction Action = TargetedBuildIdAction.load(std::memory_order_relaxed);
+	if (Action == EBuildArtifactAction::None)
+	{
+		return;
+	}
+	FCbObjectId BuildId = TargetedBuildId;
+
+	for (FBuildSelectionBuildGroupPtr BuildGroup : BuildGroups)
+	{
+		for (const TPair<FString, FArtifact>& NamedArtifact : BuildGroup->NamedArtifacts)
+		{
+			if (NamedArtifact.Value.BuildRecord.BuildId == BuildId)
+			{
+				BuildGroupListView->ClearSelection();
+				BuildGroupListView->SetItemSelection(BuildGroup, true);
+				BuildGroupListView->RequestScrollIntoView(BuildGroup);
+
+				if (Action == EBuildArtifactAction::HighlightAndDownload)
+				{
+					DownloadItem(*BuildGroup, NamedArtifact.Key, NamedArtifact.Value);
+				}
+				return;
+			}
+		}
+	}
+}
+
 SBuildSelection::EBuildType SBuildSelection::GetSelectedBuildType() const
 {
 	if (!SelectedBuildType)
@@ -485,6 +573,7 @@ void SBuildSelection::RebuildLists()
 							ExecuteOnGameThread(UE_SOURCE_LOCATION,
 								[this]
 								{
+									EnforceTargetedBuildActionOnGroup();
 									PostViewFilterChange();
 									BuildListRefreshesInProgress.fetch_sub(1);
 								});
@@ -817,6 +906,25 @@ bool SBuildSelection::DoesPassAllViewFilters(TSharedPtr<FBuildGroup> InItem) con
 		return false;
 	}
 
+	// If a row contains a targeted build and is selected, it remains visible regardless of view filters
+	for (const TPair<FString, FArtifact>& NamedArtifact : InItem->NamedArtifacts)
+	{
+		if (NamedArtifact.Value.BuildRecord.BuildId == TargetedBuildId)
+		{
+			TArray<FBuildSelectionBuildGroupPtr> SelectedItems = BuildGroupListView->GetSelectedItems();
+			if (!SelectedItems.IsEmpty())
+			{
+				for (const FBuildSelectionBuildGroupPtr& SelectedItem : SelectedItems)
+				{
+					if (SelectedItem == InItem)
+					{
+						return true;
+					}
+				}
+			}
+		}
+	}
+
 	for (const FString& ActivePlatformFilter : ActivePlatformFilters)
 	{
 		bool bHasPlatform = false;
@@ -948,7 +1056,12 @@ TSharedPtr<SWidget> SBuildSelection::OnGetBuildGroupContextMenuContent() const
 
 void SBuildSelection::RefreshArtifactSelectionWidget(FBuildSelectionBuildGroupPtr Item)
 {
-	BuildArtifactSelection->Refresh(Item ? &Item->NamedArtifacts : nullptr);
+	FCbObjectId BuildIdToHighlight;
+	if (Item && (TargetedBuildIdAction.load(std::memory_order_relaxed) != EBuildArtifactAction::None))
+	{
+		BuildIdToHighlight = TargetedBuildId.exchange(FCbObjectId());
+	}
+	BuildArtifactSelection->Refresh(Item ? &Item->NamedArtifacts : nullptr, BuildIdToHighlight);
 }
 
 void SBuildSelection::BuildGroupSelectionChanged(FBuildSelectionBuildGroupPtr Item, ESelectInfo::Type SelectInfo)
@@ -1346,105 +1459,120 @@ void SBuildSelection::DownloadSelectedArtifactsInSelectedBuildGroup()
 					continue;
 				}
 
-				FString DestinationPlatformName = ArtifactForSelectedGroup;
-				if (FCbFieldView CookPlatformField = Artifact->BuildRecord.Metadata["cookPlatform"]; CookPlatformField.HasValue() && !CookPlatformField.HasError())
-				{
-					DestinationPlatformName = *WriteToString<64>(CookPlatformField.AsString());
-				}
+				DownloadItem(*SelectedItems[0], ArtifactForSelectedGroup, *Artifact);
+			}
+		}
+	}
+}
 
-				FName TargetPlatformName = NAME_None;
-				if (FCbFieldView PlatformField = Artifact->BuildRecord.Metadata["platform"]; PlatformField.HasValue() && !PlatformField.HasError())
-				{
-					TargetPlatformName = FName(PlatformField.AsString());
-				}
+bool SBuildSelection::DownloadItem(const FBuildGroup& Group, const FString& ArtifactName, const UE::BuildSelection::Internal::FArtifact& Artifact)
+{
+	using namespace UE::Zen::Build;
 
-				EBuildType BuildType = GetSelectedBuildType();
-				int BuildTypeIndex = (int)BuildType;
-				const bool bAppendBuildNameToDestination = bAppendBuildNameToDestinations[BuildTypeIndex];
+	TSharedPtr<FBuildServiceInstance> ServiceInstance = BuildServiceInstance.Get();
+	if (!ServiceInstance)
+	{
+		return false;
+	}
 
-				FString TransferName = FString::Printf(TEXT("%s-%s"), *SelectedItems[0]->DisplayName, *DestinationPlatformName);
+	FString DestinationPlatformName = ArtifactName;
+	if (FCbFieldView CookPlatformField = Artifact.BuildRecord.Metadata["cookPlatform"]; CookPlatformField.HasValue() && !CookPlatformField.HasError())
+	{
+		DestinationPlatformName = *WriteToString<64>(CookPlatformField.AsString());
+	}
 
-				if (BuildType == EBuildType::Oplog)
-				{
-					FString RootDir = FPaths::RootDir();
-					FString EngineDir = FPaths::EngineDir();
-					FString ProjectFilename = FUProjectDictionary::GetDefault().GetProjectPathForGame(**SelectedProject);
+	FName TargetPlatformName = NAME_None;
+	if (FCbFieldView PlatformField = Artifact.BuildRecord.Metadata["platform"]; PlatformField.HasValue() && !PlatformField.HasError())
+	{
+		TargetPlatformName = FName(PlatformField.AsString());
+	}
 
-					if (ProjectFilename.IsEmpty() && UserSelectedProjectDictionary)
-					{
-						RootDir = UserSelectedProjectDictionaryRootDir;
-						EngineDir = FPaths::Combine(RootDir, TEXT("Engine"));
-						ProjectFilename = UserSelectedProjectDictionary->GetProjectPathForGame(**SelectedProject);
-					}
+	EBuildType BuildType = GetSelectedBuildType();
+	int BuildTypeIndex = (int)BuildType;
+	const bool bAppendBuildNameToDestination = bAppendBuildNameToDestinations[BuildTypeIndex];
 
-					if (ProjectFilename.IsEmpty() && FallbackProjectDictionary)
-					{
-						RootDir = FallbackProjectDictionaryRootDir;
-						EngineDir = FPaths::Combine(RootDir, TEXT("Engine"));
-						ProjectFilename = FallbackProjectDictionary->GetProjectPathForGame(**SelectedProject);
-					}
-					FString DestinationProjectId = GetEffectiveDestination();
-					FString DestinationOplogId = DestinationPlatformName;
-					if (bAppendBuildNameToDestination)
-					{
-						DestinationOplogId = FString::Printf(TEXT("%s.%s"), *DestinationOplogId, *SanitizeForZenId(SelectedItems[0]->DisplayName));
-					}
+	FString TransferName = FString::Printf(TEXT("%s-%s"), *Group.DisplayName, *DestinationPlatformName);
 
-					FString HostOverride;
-					EBuildTransferRequestFlags RequestFlags = EBuildTransferRequestFlags::Clean;
-					PreOplogBuildTransfer.ExecuteIfBound(HostOverride, RequestFlags);
-
-					FBuildServiceInstance::FBuildTransfer BuildTransfer =
-						ServiceInstance->StartOplogBuildTransfer(
-							Artifact->BuildRecord.BuildId,
-							TransferName,
-							DestinationProjectId,
-							DestinationOplogId,
-							RootDir,
-							EngineDir,
-							ProjectFilename,
-							SelectedItems[0]->Namespace,
-							Artifact->BuildRecord.BucketId,
-							TargetPlatformName,
-							HostOverride,
-							RequestFlags);
-					OnBuildTransferStarted.ExecuteIfBound(BuildTransfer, SelectedItems[0]->DisplayName, ArtifactForSelectedGroup);
-				}
-				else
-				{
-					FString DestinationFolder;
-					if (bAppendBuildNameToDestination)
-					{
-						DestinationFolder = FPaths::Combine(GetEffectiveDestination(), SanitizeForPath(SelectedItems[0]->DisplayName), DestinationPlatformName);
-					}
-					else
-					{
-						DestinationFolder = FPaths::Combine(GetEffectiveDestination(), DestinationPlatformName);
-					}
+	if (BuildType == EBuildType::Oplog)
+	{
+		FString RootDir = FPaths::RootDir();
+		FString EngineDir = FPaths::EngineDir();
+		FString ProjectFilename = FUProjectDictionary::GetDefault().GetProjectPathForGame(**SelectedProject);
 
-					FString HostOverride;
-					EBuildTransferRequestFlags RequestFlags = EBuildTransferRequestFlags::Scavenge;
-					PreBuildTransfer.ExecuteIfBound(HostOverride, RequestFlags);
-
-					FBuildServiceInstance::FBuildTransfer BuildTransfer =
-						ServiceInstance->StartBuildTransfer(
-							Artifact->BuildRecord.BuildId,
-							TransferName,
-							DestinationFolder,
-							SelectedItems[0]->Namespace,
-							Artifact->BuildRecord.BucketId,
-							UserSelectedTransferFilters[BuildTypeIndex].IncludeFilter,
-							UserSelectedTransferFilters[BuildTypeIndex].ExcludeFilter,
-							FString(),
-							TArray<FString>(),
-							TArray<FString>(),
-							HostOverride,
-							RequestFlags);
-					OnBuildTransferStarted.ExecuteIfBound(BuildTransfer, SelectedItems[0]->DisplayName, ArtifactForSelectedGroup);
-				}
-			}
+		if (ProjectFilename.IsEmpty() && UserSelectedProjectDictionary)
+		{
+			RootDir = UserSelectedProjectDictionaryRootDir;
+			EngineDir = FPaths::Combine(RootDir, TEXT("Engine"));
+			ProjectFilename = UserSelectedProjectDictionary->GetProjectPathForGame(**SelectedProject);
+		}
+
+		if (ProjectFilename.IsEmpty() && FallbackProjectDictionary)
+		{
+			RootDir = FallbackProjectDictionaryRootDir;
+			EngineDir = FPaths::Combine(RootDir, TEXT("Engine"));
+			ProjectFilename = FallbackProjectDictionary->GetProjectPathForGame(**SelectedProject);
+		}
+		FString DestinationProjectId = GetEffectiveDestination();
+		FString DestinationOplogId = DestinationPlatformName;
+		if (bAppendBuildNameToDestination)
+		{
+			DestinationOplogId = FString::Printf(TEXT("%s.%s"), *DestinationOplogId, *SanitizeForZenId(Group.DisplayName));
 		}
+
+		FString HostOverride;
+		EBuildTransferRequestFlags RequestFlags = EBuildTransferRequestFlags::Clean;
+		PreOplogBuildTransfer.ExecuteIfBound(HostOverride, RequestFlags);
+
+		FBuildServiceInstance::FBuildTransfer BuildTransfer =
+			ServiceInstance->StartOplogBuildTransfer(
+				Artifact.BuildRecord.BuildId,
+				TransferName,
+				DestinationProjectId,
+				DestinationOplogId,
+				RootDir,
+				EngineDir,
+				ProjectFilename,
+				Group.Namespace,
+				Artifact.BuildRecord.BucketId,
+				TargetPlatformName,
+				HostOverride,
+				RequestFlags);
+		OnBuildTransferStarted.ExecuteIfBound(BuildTransfer, Group.DisplayName, ArtifactName);
 	}
+	else
+	{
+		FString DestinationFolder;
+		if (bAppendBuildNameToDestination)
+		{
+			DestinationFolder = FPaths::Combine(GetEffectiveDestination(), SanitizeForPath(Group.DisplayName), DestinationPlatformName);
+		}
+		else
+		{
+			DestinationFolder = FPaths::Combine(GetEffectiveDestination(), DestinationPlatformName);
+		}
+
+		FString HostOverride;
+		EBuildTransferRequestFlags RequestFlags = EBuildTransferRequestFlags::Scavenge;
+		PreBuildTransfer.ExecuteIfBound(HostOverride, RequestFlags);
+
+		FBuildServiceInstance::FBuildTransfer BuildTransfer =
+			ServiceInstance->StartBuildTransfer(
+				Artifact.BuildRecord.BuildId,
+				TransferName,
+				DestinationFolder,
+				Group.Namespace,
+				Artifact.BuildRecord.BucketId,
+				UserSelectedTransferFilters[BuildTypeIndex].IncludeFilter,
+				UserSelectedTransferFilters[BuildTypeIndex].ExcludeFilter,
+				FString(),
+				TArray<FString>(),
+				TArray<FString>(),
+				HostOverride,
+				RequestFlags);
+		OnBuildTransferStarted.ExecuteIfBound(BuildTransfer, Group.DisplayName, ArtifactName);
+	}
+
+	return true;
 }
 
 void SBuildSelection::OnChooseDestinationDirectoryClicked()
diff --git a/Engine/Source/Programs/BuildStorageTool/Private/BuildStorageTool.cpp b/Engine/Source/Programs/BuildStorageTool/Private/BuildStorageTool.cpp
--- a/Engine/Source/Programs/BuildStorageTool/Private/BuildStorageTool.cpp
+++ b/Engine/Source/Programs/BuildStorageTool/Private/BuildStorageTool.cpp
@@ -38,6 +38,8 @@
 #include "OutputLogModule.h"
 #include "OutputLogSettings.h"
 #include "StudioTelemetry.h"
+#include "ZenBuildUtils.h"
+#include "BuildStorageToolIPC.h"
 #include "String/ParseTokens.h"
 
 
@@ -424,6 +426,25 @@ class FBuildStorageToolApp
 		return MenuBuilder.MakeWidget();
 	}
 
+	void OnIPCMessageReceived(const UE::BuildStorageTool::IMessage& Message)
+	{
+		using namespace UE::BuildStorageTool;
+
+		if (Message.GetMessageType() == FUrlMessage::MessageType)
+		{
+			const FUrlMessage& UrlMessage = static_cast<const FUrlMessage&>(Message);
+			UE::Zen::Build::FBuildReference BuildReference;
+			if (TryParseBuildReferenceFromUrl(UrlMessage.Url, BuildReference))
+			{
+				BuildSelection->ActOnBuildArtifact(
+					UrlMessage.Action == EUrlAction::HighlightAndDownload ?
+						SBuildSelection::EBuildArtifactAction::HighlightAndDownload :
+						SBuildSelection::EBuildArtifactAction::Highlight,
+					BuildReference);
+			}
+		}
+	}
+
 public:
 	FBuildStorageToolApp(FSlateApplication& InSlate, const FBuildStorageToolParameters& InToolParameters)
 		: Slate(InSlate), ToolParameters(InToolParameters)
@@ -431,10 +452,13 @@ public:
 		ZenServiceInstanceManager = MakeShared<UE::Zen::FServiceInstanceManager>();
 		BuildServiceInstanceManager = MakeShared<UE::Zen::Build::FServiceInstanceManager>();
 		InstallMessageDialogOverride();
+		MessageListenServer.OnMessageReceived.AddRaw(this, &FBuildStorageToolApp::OnIPCMessageReceived);
+		MessageListenServer.Start();
 	}
 
 	virtual ~FBuildStorageToolApp()
 	{
+		MessageListenServer.ShutdownJoin();
 		RemoveMessageDialogOverride();
 	}
 
@@ -484,7 +508,7 @@ public:
 					.HAlign(HAlign_Fill)
 					.VAlign(VAlign_Fill)
 					[
-						SNew(SBuildSelection)
+						SAssignNew(BuildSelection, SBuildSelection)
 						.ZenServiceInstance(ZenServiceInstanceManager.ToSharedRef(), &UE::Zen::FServiceInstanceManager::GetZenServiceInstance)
 						.BuildServiceInstance(BuildServiceInstanceManager.ToSharedRef(), &UE::Zen::Build::FServiceInstanceManager::GetBuildServiceInstance)
 						.PreBuildTransfer_Lambda([this]
@@ -552,6 +576,31 @@ public:
 
 		Slate.AddWindow(Window.ToSharedRef(), true);
 
+		// If the tool was launched with a Url argument, attempt to download or highlight that artifact
+		FString Url;
+		if (FParse::Value(FCommandLine::Get(), TEXT("-DownloadBuild"), Url))
+		{
+			UE::Zen::Build::FBuildReference BuildReference;
+			if (UE::Zen::Build::TryParseBuildReferenceFromUrl(Url, BuildReference))
+			{
+				BuildSelection->ActOnBuildArtifact(
+					SBuildSelection::EBuildArtifactAction::HighlightAndDownload,
+					BuildReference,
+					false);
+			}
+		}
+		else if (FParse::Value(FCommandLine::Get(), TEXT("-GoToBuild"), Url))
+		{
+			UE::Zen::Build::FBuildReference BuildReference;
+			if (UE::Zen::Build::TryParseBuildReferenceFromUrl(Url, BuildReference))
+			{
+				BuildSelection->ActOnBuildArtifact(
+					SBuildSelection::EBuildArtifactAction::Highlight,
+					BuildReference,
+					false);
+			}
+		}
+
 		// Setting focus seems to have to happen after the Window has been added
 		Slate.ClearKeyboardFocus(EFocusCause::Cleared);
 
@@ -685,6 +734,7 @@ private:
 	}
 
 	FCriticalSection CriticalSection;
+	UE::BuildStorageTool::FMessageListenServer MessageListenServer;
 	FSlateApplication& Slate;
 	TSharedPtr<SWindow> Window;
 	TSharedPtr<SNotificationItem> CompileNotification;
@@ -699,11 +749,14 @@ private:
 	bool bBoostWorkers = false;
 
 	std::atomic<bool> bLatentExclusiveOperationActive = false;
+	TSharedPtr<SBuildSelection> BuildSelection;
 	TSharedPtr<SBuildActivity> BuildActivity;
 };
 
 int BuildStorageToolMain(const TCHAR* CmdLine)
 {
+	using namespace UE::BuildStorageTool;
+
 	ON_SCOPE_EXIT
 	{
 		RequestEngineExit(TEXT("Exiting"));
@@ -726,6 +779,27 @@ int BuildStorageToolMain(const TCHAR* CmdLine)
 	FSystemWideCriticalSection SystemWideBuildStorageToolCritSec(TEXT("BuildStorageTool"));
 	if (!SystemWideBuildStorageToolCritSec.IsValid())
 	{
+		FString Url;
+		if (FParse::Value(FCommandLine::Get(), TEXT("-DownloadBuild"), Url))
+		{
+			FMessageClient MessageClient;
+			MessageClient.Start();
+
+			FUrlMessage UrlMessage;
+			UrlMessage.Action = EUrlAction::HighlightAndDownload;
+			UrlMessage.Url = Url;
+			MessageClient.SendSync(UrlMessage);
+		}
+		else if (FParse::Value(FCommandLine::Get(), TEXT("-GoToBuild"), Url))
+		{
+			FMessageClient MessageClient;
+			MessageClient.Start();
+
+			FUrlMessage UrlMessage;
+			UrlMessage.Action = EUrlAction::Highlight;
+			UrlMessage.Url = Url;
+			MessageClient.SendSync(UrlMessage);
+		}
 		return true;
 	}
 	check(GConfig && GConfig->IsReadyForUse());
diff --git a/Engine/Source/Programs/BuildStorageTool/Private/BuildStorageToolIPC.cpp b/Engine/Source/Programs/BuildStorageTool/Private/BuildStorageToolIPC.cpp
--- a/Engine/Source/Programs/BuildStorageTool/Private/BuildStorageToolIPC.cpp
+++ b/Engine/Source/Programs/BuildStorageTool/Private/BuildStorageToolIPC.cpp
@@ -288,11 +288,6 @@ void FMessageListenServer::TickCommunication()
 			ClientState.Socket = nullptr;
 		}
 
-		if (!OnMessageReceived.IsBound())
-		{
-			continue;
-		}
-
 		for (const FMarshalledMessage& Message : Messages)
 		{
 			if (Message.MessageType == FUrlMessage::MessageType)
@@ -303,7 +298,11 @@ void FMessageListenServer::TickCommunication()
 					ExecuteOnGameThread(UE_SOURCE_LOCATION,
 						[this, UrlMessage = MoveTemp(UrlMessage)]
 						{
-							OnMessageReceived.Broadcast(UrlMessage);
+
+							if (OnMessageReceived.IsBound())
+							{
+								OnMessageReceived.Broadcast(UrlMessage);
+							}
 						});
 				}
 			}
@@ -496,6 +495,59 @@ bool FMessageClient::SendSync(const IMessage& Message)
 	return SendAll(*Socket, SendBuffer.GetData(), static_cast<int32>(SendBuffer.GetSize()));
 }
 
+template <typename CharType>
+static TStringBuilderBase<CharType>& UrlActionToString(TStringBuilderBase<CharType>& Builder, const EUrlAction Action)
+{
+	switch (Action)
+	{
+	case EUrlAction::Highlight:               return Builder << ANSITEXTVIEW("Highlight");
+	case EUrlAction::HighlightAndDownload:    return Builder << ANSITEXTVIEW("HighlightAndDownload");
+	}
+	return Builder << ANSITEXTVIEW("Unknown");
+}
+
+template <typename CharType>
+static bool UrlActionFromString(EUrlAction& OutAction, const TStringView<CharType> String)
+{
+	const auto ConvertedString = StringCast<UTF8CHAR, 16>(String.GetData(), String.Len());
+	if (ConvertedString == UTF8TEXTVIEW("Highlight"))
+	{
+		OutAction = EUrlAction::Highlight;
+	}
+	else if (ConvertedString == UTF8TEXTVIEW("HighlightAndDownload"))
+	{
+		OutAction = EUrlAction::HighlightAndDownload;
+	}
+	else
+	{
+		return false;
+	}
+	return true;
+}
+
+FAnsiStringBuilderBase& operator<<(FAnsiStringBuilderBase& Builder, EUrlAction Action) { return UrlActionToString(Builder, Action); }
+FWideStringBuilderBase& operator<<(FWideStringBuilderBase& Builder, EUrlAction Action) { return UrlActionToString(Builder, Action); }
+FUtf8StringBuilderBase& operator<<(FUtf8StringBuilderBase& Builder, EUrlAction Action) { return UrlActionToString(Builder, Action); }
+
+bool TryLexFromString(EUrlAction& OutAction, FUtf8StringView String) { return UrlActionFromString(OutAction, String); }
+bool TryLexFromString(EUrlAction& OutAction, FWideStringView String) { return UrlActionFromString(OutAction, String); }
+
+FCbWriter& operator<<(FCbWriter& Writer, const EUrlAction Action)
+{
+	Writer.AddString(WriteToUtf8String<16>(Action));
+	return Writer;
+}
+
+bool LoadFromCompactBinary(FCbFieldView Field, EUrlAction& OutAction)
+{
+	if (TryLexFromString(OutAction, Field.AsString()))
+	{
+		return true;
+	}
+	OutAction = {};
+	return false;
+}
+
 FGuid FUrlMessage::MessageType(TEXT("FC6BC9307BAB4D4C9F0C236FDEF2BE62"));
 
 FUrlMessage::FUrlMessage()
@@ -504,12 +556,16 @@ FUrlMessage::FUrlMessage()
 
 void FUrlMessage::Write(FCbWriter& Writer) const
 {
+	Writer << "A" << Action;
 	Writer << "U" << Url;
 }
 
 bool FUrlMessage::TryRead(FCbObjectView Object)
 {
-	return LoadFromCompactBinary(Object["U"], Url);
+	bool bOk = true;
+	bOk = LoadFromCompactBinary(Object["A"], Action) & bOk;
+	bOk = LoadFromCompactBinary(Object["U"], Url) & bOk;
+	return bOk;
 }
 
 } // namespace UE::BuildStorageTool
diff --git a/Engine/Source/Programs/BuildStorageTool/Private/BuildStorageToolIPC.h b/Engine/Source/Programs/BuildStorageTool/Private/BuildStorageToolIPC.h
--- a/Engine/Source/Programs/BuildStorageTool/Private/BuildStorageToolIPC.h
+++ b/Engine/Source/Programs/BuildStorageTool/Private/BuildStorageToolIPC.h
@@ -112,6 +112,19 @@ private:
 	bool bReady = false;
 };
 
+enum class EUrlAction : uint8
+{
+	Highlight,
+	HighlightAndDownload,
+};
+
+FAnsiStringBuilderBase& operator<<(FAnsiStringBuilderBase& Builder, EUrlAction Action);
+FWideStringBuilderBase& operator<<(FWideStringBuilderBase& Builder, EUrlAction Action);
+FUtf8StringBuilderBase& operator<<(FUtf8StringBuilderBase& Builder, EUrlAction Action);
+
+bool TryLexFromString(EUrlAction& OutAction, FUtf8StringView String);
+bool TryLexFromString(EUrlAction& OutAction, FWideStringView String);
+
 struct FUrlMessage : public IMessage
 {
 public:
@@ -122,6 +135,7 @@ public:
 	virtual const TCHAR* GetDebugName() const override { return TEXT("URLMessage"); }
 
 public:
+	EUrlAction Action;
 	FString Url;
 	static FGuid MessageType;
 };
