BuildStorageTool: Change artifact panel to use a combobox/context menu instead of icons so that more artifacts actions can be offered.

#Synced-CL 51529686

diff --git a/Engine/Source/Editor/StorageServerWidgets/Internal/SBuildArtifactSelection.h b/Engine/Source/Editor/StorageServerWidgets/Internal/SBuildArtifactSelection.h
--- a/Engine/Source/Editor/StorageServerWidgets/Internal/SBuildArtifactSelection.h
+++ b/Engine/Source/Editor/StorageServerWidgets/Internal/SBuildArtifactSelection.h
@@ -46,6 +46,7 @@ private:
 	void RebuildSelectedGroupArtifactHeader();
 	void RebuildSelectedGroupArtifactContents();
 	void SetSelectedConfiguration(const TSharedPtr<FString> InSelectedConfiguration);
+	TSharedRef<SWidget> CreateArtifactMenuWidget(UE::BuildSelection::Internal::FArtifact* Artifact);
 
 	TMultiMap<FString, FString> ArtifactNameMap;
 	TSet<FString> ArtifactConfigurationSet;
diff --git a/Engine/Source/Editor/StorageServerWidgets/Private/SBuildArtifactSelection.cpp b/Engine/Source/Editor/StorageServerWidgets/Private/SBuildArtifactSelection.cpp
--- a/Engine/Source/Editor/StorageServerWidgets/Private/SBuildArtifactSelection.cpp
+++ b/Engine/Source/Editor/StorageServerWidgets/Private/SBuildArtifactSelection.cpp
@@ -358,46 +358,14 @@ void SBuildArtifactSelection::RebuildSelectedGroupArtifactContents()
 				.AutoWidth()
 				.VAlign(VAlign_Center)
 				[
-					SNew(SButton)
-					.ButtonStyle(&FAppStyle::Get().GetWidgetStyle<FButtonStyle>("SimpleButton"))
-					.ToolTipText(LOCTEXT("BuildSelection_BrowseToBuildUrl", "Browse to build url"))
-					.IsEnabled_Lambda([this, ArtifactName, IsBuildGroupArtifactEnabled]
-					{
-						if (!IsBuildGroupArtifactEnabled(ArtifactName))
-						{
-							return false;
-						}
-						if (FArtifact* Artifact = NamedArtifacts.Find(ArtifactName))
-						{
-							if (FCbFieldView BuildUrlField = Artifact->BuildRecord.Metadata["buildurl"]; BuildUrlField.HasValue() && !BuildUrlField.HasError())
-							{
-								return true;
-							}
-						}
-						return false;
-					})
-					.OnClicked_Lambda([this, ArtifactName, IsBuildGroupArtifactEnabled]
-					{
-						if (!IsBuildGroupArtifactEnabled(ArtifactName))
+					SNew(SComboButton)
+					.ComboButtonStyle(&FAppStyle::Get().GetWidgetStyle<FComboButtonStyle>("Artifact.SmallContext"))
+					.ForegroundColor(FSlateColor::UseStyle())
+					.OnGetMenuContent_Lambda([this, ArtifactName, IsBuildGroupArtifactEnabled]
 						{
-							return FReply::Unhandled();
-						}
-						if (FArtifact* Artifact = NamedArtifacts.Find(ArtifactName))
-						{
-							if (FCbFieldView BuildUrlField = Artifact->BuildRecord.Metadata["buildurl"]; BuildUrlField.HasValue() && !BuildUrlField.HasError())
-							{
-								FString Url(FUTF8ToTCHAR(BuildUrlField.AsString()));
-								FPlatformProcess::LaunchURL(*Url, nullptr, nullptr);
-								return FReply::Handled();
-							}
-						}
-						return FReply::Unhandled();
-					})
-					[
-						SNew(SImage)
-						.Image(FAppStyle::Get().GetBrush("Zen.BrowserView"))
-						.ColorAndOpacity(FSlateColor::UseForeground())
-					]
+							return CreateArtifactMenuWidget(NamedArtifacts.Find(ArtifactName));
+						})
+					.HasDownArrow(true)
 				]
 			];
 
@@ -417,4 +385,42 @@ void SBuildArtifactSelection::SetSelectedConfiguration(const TSharedPtr<FString>
 	WriteSetting(TEXT("Configuration"), *SelectedConfiguration);
 }
 
+TSharedRef<SWidget> SBuildArtifactSelection::CreateArtifactMenuWidget(UE::BuildSelection::Internal::FArtifact* Artifact)
+{
+	const bool bCloseAfterSelection = true;
+	const bool bCloseSelfOnly = false;
+	const bool bSearchable = false;
+	const bool bRecursivelySearchable = false;
+
+	FMenuBuilder MenuBuilder(bCloseAfterSelection,
+		nullptr,
+		TSharedPtr<FExtender>(),
+		bCloseSelfOnly,
+		&FCoreStyle::Get(),
+		bSearchable,
+		NAME_None,
+		bRecursivelySearchable);
+
+	MenuBuilder.AddMenuEntry(
+		LOCTEXT("BuildSelection_ArtifactBrowseUrl", "Browse Build URL"),
+		FText(),
+		FSlateIcon(),
+		FUIAction(
+			FExecuteAction::CreateLambda([this, Artifact]
+				{
+					if (FCbFieldView BuildUrlField = Artifact->BuildRecord.Metadata["buildurl"]; BuildUrlField.HasValue() && !BuildUrlField.HasError())
+					{
+						FString Url(FUTF8ToTCHAR(BuildUrlField.AsString()));
+						FPlatformProcess::LaunchURL(*Url, nullptr, nullptr);
+					}
+				})
+		),
+		NAME_None,
+		EUserInterfaceActionType::Button
+	);
+
+
+	return MenuBuilder.MakeWidget();
+}
+
 #undef LOCTEXT_NAMESPACE
diff --git a/Engine/Source/Programs/BuildStorageTool/Private/BuildStorageToolStyle.cpp b/Engine/Source/Programs/BuildStorageTool/Private/BuildStorageToolStyle.cpp
--- a/Engine/Source/Programs/BuildStorageTool/Private/BuildStorageToolStyle.cpp
+++ b/Engine/Source/Programs/BuildStorageTool/Private/BuildStorageToolStyle.cpp
@@ -10,6 +10,7 @@
 #include "Styling/SlateStyleRegistry.h"
 #include "Styling/SlateTypes.h"
 #include "Styling/CoreStyle.h"
+#include "Styling/StyleColors.h"
 
 #define EDITOR_IMAGE_BRUSH(RelativePath, ...) IMAGE_BRUSH("../Editor/Slate/" RelativePath, __VA_ARGS__)
 #define EDITOR_IMAGE_BRUSH_SVG(RelativePath, ...) IMAGE_BRUSH_SVG("../Editor/Slate/" RelativePath, __VA_ARGS__)
@@ -70,6 +71,24 @@ TSharedRef< FSlateStyleSet > FBuildStorageToolStyle::Create()
 
 	Style.Set("Zen.Cancel", new CORE_IMAGE_BRUSH_SVG("Starship/Common/reject", Icon12x12));
 
+	const FButtonStyle NoPaddingSimpleButton = FButtonStyle()
+	.SetNormal(FSlateNoResource())
+	.SetHovered(FSlateRoundedBoxBrush(FStyleColors::Dropdown, 4.0f))
+	.SetPressed(FSlateRoundedBoxBrush(FStyleColors::Dropdown, 4.0f))
+	.SetDisabled(FSlateNoResource())
+	.SetNormalForeground(FStyleColors::Foreground)
+	.SetHoveredForeground(FStyleColors::ForegroundHover)
+	.SetPressedForeground(FStyleColors::Primary)
+	.SetDisabledForeground(FStyleColors::Foreground)
+	.SetNormalPadding(0)
+	.SetPressedPadding(0);
+	FComboButtonStyle SmallContextCombo = FComboButtonStyle(Style.GetWidgetStyle<FComboButtonStyle>("ComboButton"))
+		.SetContentPadding(FMargin(3.f, 0.f))
+		.SetButtonStyle(NoPaddingSimpleButton)
+		.SetDownArrowImage(IMAGE_BRUSH_SVG("Starship/Common/ellipsis-vertical-narrow", FVector2f(4.f, 16.f)));
+
+	Style.Set("Artifact.SmallContext", SmallContextCombo);
+
 	const FTextBlockStyle DefaultText = FTextBlockStyle()
 		.SetFont(DEFAULT_FONT("Bold", 10))
 		.SetColorAndOpacity(FSlateColor::UseForeground())
