BuildStorageTool: Added metadata viewing for build artifacts

#Synced-CL 51529686

diff --git a/Engine/Source/Editor/StorageServerWidgets/Internal/SBuildArtifactSelection.h b/Engine/Source/Editor/StorageServerWidgets/Internal/SBuildArtifactSelection.h
--- a/Engine/Source/Editor/StorageServerWidgets/Internal/SBuildArtifactSelection.h
+++ b/Engine/Source/Editor/StorageServerWidgets/Internal/SBuildArtifactSelection.h
@@ -41,12 +41,12 @@ public:
 			TSet<FString>&& InArtifactConfigurationSet,
 			TMultiMap<FString, FString>&& InArtifactNameMap
 		);
-	UE_API void Refresh(const TMap<FString, UE::BuildSelection::Internal::FArtifact>* InNamedArtifacts, FCbObjectId InHighlightBuildId = {});
+	UE_API void Refresh(FStringView InGroupDisplayName, const TMap<FString, UE::BuildSelection::Internal::FArtifact>* InNamedArtifacts, FCbObjectId InHighlightBuildId = {});
 private:
 	void RebuildSelectedGroupArtifactHeader();
 	void RebuildSelectedGroupArtifactContents();
 	void SetSelectedConfiguration(const TSharedPtr<FString> InSelectedConfiguration);
-	TSharedRef<SWidget> CreateArtifactMenuWidget(UE::BuildSelection::Internal::FArtifact* Artifact);
+	TSharedRef<SWidget> CreateArtifactMenuWidget(UE::BuildSelection::Internal::FArtifact* Artifact, const FString& ArtifactName);
 
 	TMultiMap<FString, FString> ArtifactNameMap;
 	TSet<FString> ArtifactConfigurationSet;
@@ -54,6 +54,7 @@ private:
 	TSharedPtr<FString> SelectedConfiguration;
 	TArray<FString> SelectedArtifacts;
 	FCbObjectId HighlightedBuildId;
+	FString GroupDisplayName;
 
 	TMap<FString, UE::BuildSelection::Internal::FArtifact> NamedArtifacts;
 
diff --git a/Engine/Source/Editor/StorageServerWidgets/Private/SBuildArtifactMetadataViewer.cpp b/Engine/Source/Editor/StorageServerWidgets/Private/SBuildArtifactMetadataViewer.cpp
new file mode 100644
--- /dev/null
+++ b/Engine/Source/Editor/StorageServerWidgets/Private/SBuildArtifactMetadataViewer.cpp
@@ -0,0 +1,347 @@
+// Copyright Epic Games, Inc. All Rights Reserved.
+
+#include "SBuildArtifactMetadataViewer.h"
+#include "Widgets/Text/STextBlock.h"
+#include "Widgets/Layout/SScrollBox.h"
+#include "Widgets/Layout/SBorder.h"
+#include "Widgets/Layout/SSeparator.h"
+#include "Widgets/Views/STableRow.h"
+#include "Widgets/Input/SSearchBox.h"
+#include "Styling/StyleColors.h"
+#include "Styling/AppStyle.h"
+#include "Framework/MultiBox/MultiBoxBuilder.h"
+#include "HAL/PlatformApplicationMisc.h"
+
+#define LOCTEXT_NAMESPACE "BuildArtifactMetadata"
+
+void SBuildArtifactMetadataViewer::Construct(const FArguments& InArgs)
+{
+	ArtifactName = InArgs._ArtifactName;
+
+	// Parse the metadata
+	ParseMetadata(InArgs._Metadata);
+
+	// Initialize filtered list with all entries
+	RefreshFilteredList();
+
+	ChildSlot
+	[
+		SNew(SVerticalBox)
+
+		// Header
+		+ SVerticalBox::Slot()
+		.AutoHeight()
+		.Padding(8.0f)
+		[
+			SNew(STextBlock)
+			.Text(FText::Format(LOCTEXT("MetadataHeader", "Metadata for: {0}"), FText::FromString(ArtifactName)))
+			.Font(FCoreStyle::GetDefaultFontStyle("Bold", 12))
+		]
+
+		+ SVerticalBox::Slot()
+		.AutoHeight()
+		[
+			SNew(SSeparator)
+		]
+
+		// Search box
+		+ SVerticalBox::Slot()
+		.AutoHeight()
+		.Padding(4.0f)
+		[
+			SNew(SSearchBox)
+			.HintText(LOCTEXT("SearchHint", "Search metadata..."))
+			.OnTextChanged(this, &SBuildArtifactMetadataViewer::OnSearchTextChanged)
+		]
+
+		// Metadata list
+		+ SVerticalBox::Slot()
+		.FillHeight(1.0f)
+		.Padding(4.0f)
+		[
+			SNew(SBorder)
+			.BorderImage(FAppStyle::GetBrush("ToolPanel.GroupBorder"))
+			.Padding(0.0f)
+			[
+				SAssignNew(ListView, SListView<TSharedPtr<FMetadataEntry>>)
+				.ListItemsSource(&FilteredMetadataEntries)
+				.OnGenerateRow(this, &SBuildArtifactMetadataViewer::OnGenerateRow)
+				.OnContextMenuOpening(this, &SBuildArtifactMetadataViewer::OnContextMenuOpening)
+				.SelectionMode(ESelectionMode::Single)
+				.HeaderRow
+				(
+					SNew(SHeaderRow)
+					+ SHeaderRow::Column("Key")
+					.DefaultLabel(LOCTEXT("KeyColumn", "Key"))
+					.FillWidth(0.3f)
+
+					+ SHeaderRow::Column("Value")
+					.DefaultLabel(LOCTEXT("ValueColumn", "Value"))
+					.FillWidth(0.6f)
+
+					+ SHeaderRow::Column("Type")
+					.DefaultLabel(LOCTEXT("TypeColumn", "Type"))
+					.FillWidth(0.1f)
+				)
+			]
+		]
+	];
+}
+
+TSharedRef<ITableRow> SBuildArtifactMetadataViewer::OnGenerateRow(TSharedPtr<FMetadataEntry> Entry, const TSharedRef<STableViewBase>& OwnerTable)
+{
+	return SNew(STableRow<TSharedPtr<FMetadataEntry>>, OwnerTable)
+		[
+			SNew(SHorizontalBox)
+
+			// Key column
+			+ SHorizontalBox::Slot()
+			.FillWidth(0.3f)
+			.Padding(4.0f, 2.0f)
+			.VAlign(VAlign_Center)
+			[
+				SNew(STextBlock)
+				.Text(FText::FromString(Entry->Key))
+				.Font(FCoreStyle::GetDefaultFontStyle("Bold", 9))
+				.ColorAndOpacity(FStyleColors::AccentGreen)
+			]
+
+			// Value column
+			+ SHorizontalBox::Slot()
+			.FillWidth(0.6f)
+			.Padding(4.0f, 2.0f)
+			.VAlign(VAlign_Center)
+			[
+				SNew(STextBlock)
+				.Text(FText::FromString(Entry->Value))
+				.Font(FCoreStyle::GetDefaultFontStyle("Regular", 9))
+				.AutoWrapText(true)
+			]
+
+			// Type column
+			+ SHorizontalBox::Slot()
+			.FillWidth(0.1f)
+			.Padding(4.0f, 2.0f)
+			.VAlign(VAlign_Center)
+			[
+				SNew(STextBlock)
+				.Text(FText::FromString(Entry->Type))
+				.Font(FCoreStyle::GetDefaultFontStyle("Italic", 8))
+				.ColorAndOpacity(FStyleColors::Foreground.GetSpecifiedColor() * 0.6f)
+			]
+		];
+}
+
+void SBuildArtifactMetadataViewer::ParseMetadata(const FCbObject& Metadata)
+{
+	MetadataEntries.Empty();
+
+	// Iterate through all fields in the compact binary object
+	for (FCbFieldView Field : Metadata)
+	{
+		FUtf8StringView FieldName = Field.GetName();
+		FString Key = FString(FieldName.Len(), FieldName.GetData());
+		FString Value;
+		FString Type;
+
+		// Extract value based on field type
+		// Check more specific types before more general ones to avoid misidentification
+
+		// Check null first (most specific)
+		if (Field.IsNull())
+		{
+			Value = TEXT("null");
+			Type = TEXT("Null");
+		}
+		// Check bool before integers
+		else if (Field.IsBool())
+		{
+			Value = Field.AsBool() ? TEXT("true") : TEXT("false");
+			Type = TEXT("Bool");
+		}
+		// Check DateTime before integer (DateTime is stored as integer internally)
+		else if (Field.IsDateTime())
+		{
+			FDateTime DateTime = Field.AsDateTime();
+			Value = DateTime.ToString();
+			Type = TEXT("DateTime");
+		}
+		// Check ObjectId before binary (ObjectId is a specific binary type)
+		else if (Field.IsObjectId())
+		{
+			Value = FString(WriteToString<64>(Field.AsObjectId()));
+			Type = TEXT("ObjectId");
+		}
+		// Check UUID before binary (UUID is a specific binary type)
+		else if (Field.IsUuid())
+		{
+			Value = FString(WriteToString<48>(Field.AsUuid()));
+			Type = TEXT("UUID");
+		}
+		// Check Hash before binary (Hash is a specific binary type)
+		else if (Field.IsHash())
+		{
+			Value = FString(WriteToString<48>(Field.AsHash()));
+			Type = TEXT("Hash");
+		}
+		// Check integer before float (integer is more specific than float)
+		else if (Field.IsInteger())
+		{
+			Value = FString::Printf(TEXT("%lld"), Field.AsInt64());
+			Type = TEXT("Integer");
+		}
+		else if (Field.IsFloat())
+		{
+			Value = FString::Printf(TEXT("%f"), Field.AsFloat());
+			Type = TEXT("Float");
+		}
+		// Check string
+		else if (Field.IsString())
+		{
+			Value = FString(FUTF8ToTCHAR(Field.AsString()));
+			Type = TEXT("String");
+		}
+		// Check generic binary after specific binary types
+		else if (Field.IsBinary())
+		{
+			FMemoryView BinaryView = Field.AsBinaryView();
+			Value = FString::Printf(TEXT("[%lld bytes]"), BinaryView.GetSize());
+			Type = TEXT("Binary");
+		}
+		// Check container types
+		else if (Field.IsArray())
+		{
+			Value = FString::Printf(TEXT("[%d elements]"), Field.AsArrayView().Num());
+			Type = TEXT("Array");
+		}
+		else if (Field.IsObject())
+		{
+			int32 FieldCount = 0;
+			for (FCbFieldView ObjectField : Field)
+			{
+				++FieldCount;
+			}
+			Value = FString::Printf(TEXT("{%d fields}"), FieldCount);
+			Type = TEXT("Object");
+		}
+		else
+		{
+			Value = TEXT("<Unknown>");
+			Type = TEXT("Unknown");
+		}
+
+		MetadataEntries.Add(MakeShared<FMetadataEntry>(Key, Value, Type));
+	}
+
+	// Sort entries alphabetically by key
+	MetadataEntries.Sort([](const TSharedPtr<FMetadataEntry>& A, const TSharedPtr<FMetadataEntry>& B)
+	{
+		return A->Key < B->Key;
+	});
+}
+
+void SBuildArtifactMetadataViewer::RefreshFilteredList()
+{
+	FilteredMetadataEntries.Empty();
+
+	if (SearchText.IsEmpty())
+	{
+		// No filter, show all entries
+		FilteredMetadataEntries = MetadataEntries;
+	}
+	else
+	{
+		// Filter entries based on search text (case-insensitive)
+		FString SearchString = SearchText.ToString().ToLower();
+		for (const TSharedPtr<FMetadataEntry>& Entry : MetadataEntries)
+		{
+			if (Entry->Key.ToLower().Contains(SearchString) ||
+				Entry->Value.ToLower().Contains(SearchString) ||
+				Entry->Type.ToLower().Contains(SearchString))
+			{
+				FilteredMetadataEntries.Add(Entry);
+			}
+		}
+	}
+
+	// Refresh the list view
+	if (ListView.IsValid())
+	{
+		ListView->RequestListRefresh();
+	}
+}
+
+void SBuildArtifactMetadataViewer::OnSearchTextChanged(const FText& InText)
+{
+	SearchText = InText;
+	RefreshFilteredList();
+}
+
+TSharedPtr<SWidget> SBuildArtifactMetadataViewer::OnContextMenuOpening()
+{
+	TArray<TSharedPtr<FMetadataEntry>> SelectedItems = ListView->GetSelectedItems();
+	if (SelectedItems.Num() == 0)
+	{
+		return nullptr;
+	}
+
+	const bool bCloseAfterSelection = true;
+	FMenuBuilder MenuBuilder(bCloseAfterSelection, nullptr);
+
+	MenuBuilder.BeginSection("MetadataCopy", LOCTEXT("CopySection", "Copy"));
+	{
+		MenuBuilder.AddMenuEntry(
+			LOCTEXT("CopyKey", "Copy Key"),
+			LOCTEXT("CopyKeyTooltip", "Copy the metadata key to clipboard"),
+			FSlateIcon(),
+			FUIAction(FExecuteAction::CreateSP(this, &SBuildArtifactMetadataViewer::CopyKeyToClipboard))
+		);
+
+		MenuBuilder.AddMenuEntry(
+			LOCTEXT("CopyValue", "Copy Value"),
+			LOCTEXT("CopyValueTooltip", "Copy the metadata value to clipboard"),
+			FSlateIcon(),
+			FUIAction(FExecuteAction::CreateSP(this, &SBuildArtifactMetadataViewer::CopyValueToClipboard))
+		);
+
+		MenuBuilder.AddMenuEntry(
+			LOCTEXT("CopyEntry", "Copy Key: Value"),
+			LOCTEXT("CopyEntryTooltip", "Copy the entire entry as 'Key: Value' to clipboard"),
+			FSlateIcon(),
+			FUIAction(FExecuteAction::CreateSP(this, &SBuildArtifactMetadataViewer::CopyEntryToClipboard))
+		);
+	}
+	MenuBuilder.EndSection();
+
+	return MenuBuilder.MakeWidget();
+}
+
+void SBuildArtifactMetadataViewer::CopyKeyToClipboard()
+{
+	TArray<TSharedPtr<FMetadataEntry>> SelectedItems = ListView->GetSelectedItems();
+	if (SelectedItems.Num() > 0)
+	{
+		FPlatformApplicationMisc::ClipboardCopy(*SelectedItems[0]->Key);
+	}
+}
+
+void SBuildArtifactMetadataViewer::CopyValueToClipboard()
+{
+	TArray<TSharedPtr<FMetadataEntry>> SelectedItems = ListView->GetSelectedItems();
+	if (SelectedItems.Num() > 0)
+	{
+		FPlatformApplicationMisc::ClipboardCopy(*SelectedItems[0]->Value);
+	}
+}
+
+void SBuildArtifactMetadataViewer::CopyEntryToClipboard()
+{
+	TArray<TSharedPtr<FMetadataEntry>> SelectedItems = ListView->GetSelectedItems();
+	if (SelectedItems.Num() > 0)
+	{
+		FString EntryText = FString::Printf(TEXT("%s: %s"), *SelectedItems[0]->Key, *SelectedItems[0]->Value);
+		FPlatformApplicationMisc::ClipboardCopy(*EntryText);
+	}
+}
+
+#undef LOCTEXT_NAMESPACE
diff --git a/Engine/Source/Editor/StorageServerWidgets/Private/SBuildArtifactMetadataViewer.h b/Engine/Source/Editor/StorageServerWidgets/Private/SBuildArtifactMetadataViewer.h
new file mode 100644
--- /dev/null
+++ b/Engine/Source/Editor/StorageServerWidgets/Private/SBuildArtifactMetadataViewer.h
@@ -0,0 +1,76 @@
+// Copyright Epic Games, Inc. All Rights Reserved.
+
+#pragma once
+
+#include "CoreMinimal.h"
+#include "Widgets/DeclarativeSyntaxSupport.h"
+#include "Widgets/SCompoundWidget.h"
+#include "Widgets/Views/SListView.h"
+#include "Serialization/CompactBinary.h"
+
+/**
+ * Widget for displaying build artifact metadata in a property grid style
+ */
+class SBuildArtifactMetadataViewer : public SCompoundWidget
+{
+public:
+	/** Struct to hold a metadata key-value pair for display */
+	struct FMetadataEntry
+	{
+		FString Key;
+		FString Value;
+		FString Type;
+
+		FMetadataEntry(const FString& InKey, const FString& InValue, const FString& InType)
+			: Key(InKey), Value(InValue), Type(InType)
+		{}
+	};
+
+	SLATE_BEGIN_ARGS(SBuildArtifactMetadataViewer)
+	{}
+		SLATE_ARGUMENT(FCbObject, Metadata)
+		SLATE_ARGUMENT(FString, ArtifactName)
+	SLATE_END_ARGS()
+
+	void Construct(const FArguments& InArgs);
+
+private:
+	/** Generate a row widget for a metadata entry */
+	TSharedRef<ITableRow> OnGenerateRow(TSharedPtr<FMetadataEntry> Entry, const TSharedRef<STableViewBase>& OwnerTable);
+
+	/** Generate context menu for a row */
+	TSharedPtr<SWidget> OnContextMenuOpening();
+
+	/** Parse FCbObject metadata into display entries */
+	void ParseMetadata(const FCbObject& Metadata);
+
+	/** Refresh the filtered list based on search text */
+	void RefreshFilteredList();
+
+	/** Called when search text changes */
+	void OnSearchTextChanged(const FText& InText);
+
+	/** Copy selected entry key to clipboard */
+	void CopyKeyToClipboard();
+
+	/** Copy selected entry value to clipboard */
+	void CopyValueToClipboard();
+
+	/** Copy selected entry as "Key: Value" to clipboard */
+	void CopyEntryToClipboard();
+
+	/** List of all metadata entries */
+	TArray<TSharedPtr<FMetadataEntry>> MetadataEntries;
+
+	/** Filtered list of metadata entries based on search */
+	TArray<TSharedPtr<FMetadataEntry>> FilteredMetadataEntries;
+
+	/** List view widget */
+	TSharedPtr<SListView<TSharedPtr<FMetadataEntry>>> ListView;
+
+	/** Current search filter text */
+	FText SearchText;
+
+	/** Artifact name for display */
+	FString ArtifactName;
+};
diff --git a/Engine/Source/Editor/StorageServerWidgets/Private/SBuildArtifactSelection.cpp b/Engine/Source/Editor/StorageServerWidgets/Private/SBuildArtifactSelection.cpp
--- a/Engine/Source/Editor/StorageServerWidgets/Private/SBuildArtifactSelection.cpp
+++ b/Engine/Source/Editor/StorageServerWidgets/Private/SBuildArtifactSelection.cpp
@@ -3,13 +3,16 @@
 #include "SBuildArtifactSelection.h"
 
 #include "Framework/MultiBox/MultiBoxBuilder.h"
+#include "Framework/Docking/TabManager.h"
 #include "Styling/StyleColors.h"
 #include "Widgets/Input/SButton.h"
 #include "Widgets/Input/SCheckBox.h"
 #include "Widgets/Input/SComboBox.h"
 #include "Widgets/Text/STextBlock.h"
+#include "Widgets/Docking/SDockTab.h"
 #include "ZenServiceInstanceManager.h"
-#include <atomic>
+#include "SBuildArtifactMetadataViewer.h"
+
 
 #define LOCTEXT_NAMESPACE "StorageServerBuild"
 
@@ -118,11 +121,12 @@ void SBuildArtifactSelection::SetAvailableConfigurations(
 	}
 }
 
-void SBuildArtifactSelection::Refresh(const TMap<FString, UE::BuildSelection::Internal::FArtifact>* InNamedArtifacts, FCbObjectId InHighlightBuildId)
+void SBuildArtifactSelection::Refresh(const FStringView InGroupDisplayName, const TMap<FString, UE::BuildSelection::Internal::FArtifact>* InNamedArtifacts, FCbObjectId InHighlightBuildId)
 {
 	using namespace UE::BuildSelection::Internal;
 
 	HighlightedBuildId = InHighlightBuildId;
+	GroupDisplayName = InGroupDisplayName;
 	if (InNamedArtifacts)
 	{
 		NamedArtifacts = *InNamedArtifacts;
@@ -363,7 +367,7 @@ void SBuildArtifactSelection::RebuildSelectedGroupArtifactContents()
 					.ForegroundColor(FSlateColor::UseStyle())
 					.OnGetMenuContent_Lambda([this, ArtifactName, IsBuildGroupArtifactEnabled]
 						{
-							return CreateArtifactMenuWidget(NamedArtifacts.Find(ArtifactName));
+							return CreateArtifactMenuWidget(NamedArtifacts.Find(ArtifactName), ArtifactName);
 						})
 					.HasDownArrow(true)
 				]
@@ -385,7 +389,7 @@ void SBuildArtifactSelection::SetSelectedConfiguration(const TSharedPtr<FString>
 	WriteSetting(TEXT("Configuration"), *SelectedConfiguration);
 }
 
-TSharedRef<SWidget> SBuildArtifactSelection::CreateArtifactMenuWidget(UE::BuildSelection::Internal::FArtifact* Artifact)
+TSharedRef<SWidget> SBuildArtifactSelection::CreateArtifactMenuWidget(UE::BuildSelection::Internal::FArtifact* Artifact, const FString& ArtifactName)
 {
 	const bool bCloseAfterSelection = true;
 	const bool bCloseSelfOnly = false;
@@ -419,6 +423,82 @@ TSharedRef<SWidget> SBuildArtifactSelection::CreateArtifactMenuWidget(UE::BuildS
 		EUserInterfaceActionType::Button
 	);
 
+	MenuBuilder.AddMenuEntry(
+		LOCTEXT("BuildSelection_ViewMetadata", "View Metadata"),
+		FText(),
+		FSlateIcon(),
+		FUIAction(
+			FExecuteAction::CreateLambda([this, Artifact, ArtifactName]
+				{
+					if (Artifact)
+					{
+						static const FName MetadataTabId(TEXT("BuildArtifactMetadata"));
+						static TArray<TWeakPtr<SDockTab>> MetadataTabs;
+						static TSharedPtr<SDockTab> PendingTab;
+
+						// Register a spawner once; it returns the pre-built tab content
+						if (!FGlobalTabmanager::Get()->HasTabSpawner(MetadataTabId))
+						{
+							FGlobalTabmanager::Get()->RegisterNomadTabSpawner(MetadataTabId,
+								FOnSpawnTab::CreateLambda([](const FSpawnTabArgs&) -> TSharedRef<SDockTab>
+								{
+									return PendingTab.ToSharedRef();
+								}))
+								.SetMenuType(ETabSpawnerMenuType::Hidden);
+						}
+
+						// Create full display name including group
+						FString FullDisplayName = ArtifactName;
+						if (!GroupDisplayName.IsEmpty())
+						{
+							FullDisplayName = FString::Printf(TEXT("%s - %s"), *GroupDisplayName, *ArtifactName);
+						}
+
+						// Create the tab with real content
+						TSharedRef<SDockTab> NewTab = SNew(SDockTab)
+							.TabRole(ETabRole::DocumentTab)
+							.Label(FText::FromString(FullDisplayName))
+							[
+								SNew(SBuildArtifactMetadataViewer)
+								.Metadata(Artifact->BuildRecord.Metadata)
+								.ArtifactName(FullDisplayName)
+							];
+
+						// Find any existing live metadata tab
+						TSharedPtr<SDockTab> ExistingTab;
+						for (int32 i = MetadataTabs.Num() - 1; i >= 0; --i)
+						{
+							if (TSharedPtr<SDockTab> Tab = MetadataTabs[i].Pin())
+							{
+								ExistingTab = Tab;
+							}
+							else
+							{
+								MetadataTabs.RemoveAt(i);
+							}
+						}
+
+						if (ExistingTab.IsValid())
+						{
+							// Dock alongside the existing metadata tab
+							FGlobalTabmanager::Get()->InsertNewDocumentTab(
+								MetadataTabId, FTabManager::FLiveTabSearch(), NewTab);
+						}
+						else
+						{
+							// No existing tab; open in a new window via the spawner
+							PendingTab = NewTab;
+							FGlobalTabmanager::Get()->TryInvokeTab(MetadataTabId);
+							PendingTab.Reset();
+						}
+
+						MetadataTabs.Add(NewTab);
+					}
+				})
+		),
+		NAME_None,
+		EUserInterfaceActionType::Button
+	);
 
 	return MenuBuilder.MakeWidget();
 }
diff --git a/Engine/Source/Editor/StorageServerWidgets/Private/SBuildSelection.cpp b/Engine/Source/Editor/StorageServerWidgets/Private/SBuildSelection.cpp
--- a/Engine/Source/Editor/StorageServerWidgets/Private/SBuildSelection.cpp
+++ b/Engine/Source/Editor/StorageServerWidgets/Private/SBuildSelection.cpp
@@ -1061,7 +1061,7 @@ void SBuildSelection::RefreshArtifactSelectionWidget(FBuildSelectionBuildGroupPt
 	{
 		BuildIdToHighlight = TargetedBuildId.exchange(FCbObjectId());
 	}
-	BuildArtifactSelection->Refresh(Item ? &Item->NamedArtifacts : nullptr, BuildIdToHighlight);
+	BuildArtifactSelection->Refresh(Item ? Item->DisplayName : FStringView(), Item ? &Item->NamedArtifacts : nullptr, BuildIdToHighlight);
 }
 
 void SBuildSelection::BuildGroupSelectionChanged(FBuildSelectionBuildGroupPtr Item, ESelectInfo::Type SelectInfo)
