BuildStorageTool: Fix for zero commit field in build metadata view.

#Synced-CL 51529686

diff --git a/Engine/Source/Editor/StorageServerWidgets/Private/SBuildArtifactMetadataViewer.cpp b/Engine/Source/Editor/StorageServerWidgets/Private/SBuildArtifactMetadataViewer.cpp
--- a/Engine/Source/Editor/StorageServerWidgets/Private/SBuildArtifactMetadataViewer.cpp
+++ b/Engine/Source/Editor/StorageServerWidgets/Private/SBuildArtifactMetadataViewer.cpp
@@ -192,7 +192,7 @@ void SBuildArtifactMetadataViewer::ParseMetadata(const FCbObject& Metadata)
 		}
 		else if (Field.IsFloat())
 		{
-			Value = FString::Printf(TEXT("%f"), Field.AsFloat());
+			Value = FString::Printf(TEXT("%f"), Field.AsDouble());
 			Type = TEXT("Float");
 		}
 		// Check string
