BuildStorageTool: stop using atomic and use RWLock instead to avoid a Linux link issue related to libatomic.  Also addresses flaw where two fields could change independently of each other.

#Synced-CL 51529686

diff --git a/Engine/Source/Editor/StorageServerWidgets/Internal/SBuildSelection.h b/Engine/Source/Editor/StorageServerWidgets/Internal/SBuildSelection.h
--- a/Engine/Source/Editor/StorageServerWidgets/Internal/SBuildSelection.h
+++ b/Engine/Source/Editor/StorageServerWidgets/Internal/SBuildSelection.h
@@ -6,6 +6,7 @@
 #include "Experimental/BuildServerInterface.h"
 #include "Experimental/ZenServerInterface.h"
 #include "Internationalization/Regex.h"
+#include "HAL/CriticalSection.h"
 #include "Templates/SharedPointer.h"
 #include "Widgets/DeclarativeSyntaxSupport.h"
 #include "Widgets/Input/SComboBox.h"
@@ -174,8 +175,9 @@ private:
 		FString ExcludeFilter;
 	};
 
-	std::atomic<EBuildArtifactAction> TargetedBuildIdAction = EBuildArtifactAction::None;
-	std::atomic<FCbObjectId> TargetedBuildId;
+	mutable FRWLock TargetedBuildIdLock;
+	EBuildArtifactAction TargetedBuildIdAction = EBuildArtifactAction::None;
+	FCbObjectId TargetedBuildId;
 
 	TArray<FString> ActivePlatformFilters;
 	FString UserSelectedDestinations[(int)EBuildType::Count];
diff --git a/Engine/Source/Editor/StorageServerWidgets/Private/SBuildSelection.cpp b/Engine/Source/Editor/StorageServerWidgets/Private/SBuildSelection.cpp
--- a/Engine/Source/Editor/StorageServerWidgets/Private/SBuildSelection.cpp
+++ b/Engine/Source/Editor/StorageServerWidgets/Private/SBuildSelection.cpp
@@ -258,8 +258,11 @@ void SBuildSelection::ActOnBuildArtifact(EBuildArtifactAction Action, const UE::
 		bNeedListRebuild = true;
 	}
 
-	TargetedBuildIdAction.store(Action);
-	TargetedBuildId.store(BuildReference.BuildId);
+	{
+		FWriteScopeLock WriteLock(TargetedBuildIdLock);
+		TargetedBuildIdAction = Action;
+		TargetedBuildId = BuildReference.BuildId;
+	}
 
 	if (bAllowRebuildLists && bNeedListRebuild)
 	{
@@ -283,12 +286,17 @@ void SBuildSelection::EnforceTargetedBuildActionOnGroup()
 {
 	using namespace UE::BuildSelection::Internal;
 
-	EBuildArtifactAction Action = TargetedBuildIdAction.load(std::memory_order_relaxed);
+	EBuildArtifactAction Action;
+	FCbObjectId BuildId;
+	{
+		FReadScopeLock ReadLock(TargetedBuildIdLock);
+		Action = TargetedBuildIdAction;
+		BuildId = TargetedBuildId;
+	}
 	if (Action == EBuildArtifactAction::None)
 	{
 		return;
 	}
-	FCbObjectId BuildId = TargetedBuildId;
 
 	for (FBuildSelectionBuildGroupPtr BuildGroup : BuildGroups)
 	{
@@ -906,10 +914,18 @@ bool SBuildSelection::DoesPassAllViewFilters(TSharedPtr<FBuildGroup> InItem) con
 		return false;
 	}
 
+	EBuildArtifactAction Action;
+	FCbObjectId BuildId;
+	{
+		FReadScopeLock ReadLock(TargetedBuildIdLock);
+		Action = TargetedBuildIdAction;
+		BuildId = TargetedBuildId;
+	}
+
 	// If a row contains a targeted build and is selected, it remains visible regardless of view filters
 	for (const TPair<FString, FArtifact>& NamedArtifact : InItem->NamedArtifacts)
 	{
-		if (NamedArtifact.Value.BuildRecord.BuildId == TargetedBuildId)
+		if (NamedArtifact.Value.BuildRecord.BuildId == BuildId)
 		{
 			TArray<FBuildSelectionBuildGroupPtr> SelectedItems = BuildGroupListView->GetSelectedItems();
 			if (!SelectedItems.IsEmpty())
@@ -1056,10 +1072,25 @@ TSharedPtr<SWidget> SBuildSelection::OnGetBuildGroupContextMenuContent() const
 
 void SBuildSelection::RefreshArtifactSelectionWidget(FBuildSelectionBuildGroupPtr Item)
 {
+	EBuildArtifactAction Action;
+	{
+		FReadScopeLock ReadLock(TargetedBuildIdLock);
+		Action = TargetedBuildIdAction;
+	}
+
 	FCbObjectId BuildIdToHighlight;
-	if (Item && (TargetedBuildIdAction.load(std::memory_order_relaxed) != EBuildArtifactAction::None))
+	if (Item && (Action != EBuildArtifactAction::None))
 	{
-		BuildIdToHighlight = TargetedBuildId.exchange(FCbObjectId());
+		FCbObjectId BuildId;
+		{
+			FWriteScopeLock WriteLock(TargetedBuildIdLock);
+			Action = TargetedBuildIdAction;
+			if (Action != EBuildArtifactAction::None)
+			{
+				BuildIdToHighlight = TargetedBuildId;
+				TargetedBuildId = FCbObjectId();
+			}
+		}
 	}
 	BuildArtifactSelection->Refresh(Item ? Item->DisplayName : FStringView(), Item ? &Item->NamedArtifacts : nullptr, BuildIdToHighlight);
 }
