BuildStorageTool: Add the ability to view the contents of a build.

#Synced-CL 51529686

diff --git a/Engine/Source/Editor/StorageServerWidgets/Internal/SBuildArtifactSelection.h b/Engine/Source/Editor/StorageServerWidgets/Internal/SBuildArtifactSelection.h
--- a/Engine/Source/Editor/StorageServerWidgets/Internal/SBuildArtifactSelection.h
+++ b/Engine/Source/Editor/StorageServerWidgets/Internal/SBuildArtifactSelection.h
@@ -41,13 +41,17 @@ public:
 			TSet<FString>&& InArtifactConfigurationSet,
 			TMultiMap<FString, FString>&& InArtifactNameMap
 		);
-	UE_API void Refresh(FStringView InGroupDisplayName, const TMap<FString, UE::BuildSelection::Internal::FArtifact>* InNamedArtifacts, FCbObjectId InHighlightBuildId = {});
+	UE_API void Refresh(FStringView InNamespace, FStringView InGroupDisplayName, const TMap<FString, UE::BuildSelection::Internal::FArtifact>* InNamedArtifacts, FCbObjectId InHighlightBuildId = {}, bool bInCanViewContents = false);
 private:
 	void RebuildSelectedGroupArtifactHeader();
 	void RebuildSelectedGroupArtifactContents();
 	void SetSelectedConfiguration(const TSharedPtr<FString> InSelectedConfiguration);
 	TSharedRef<SWidget> CreateArtifactMenuWidget(UE::BuildSelection::Internal::FArtifact* Artifact, const FString& ArtifactName);
 
+	TAttribute<TSharedPtr<UE::Zen::Build::FBuildServiceInstance>> BuildServiceInstance;
+	FString Namespace;
+	bool bCanViewContents = false;
+
 	TMultiMap<FString, FString> ArtifactNameMap;
 	TSet<FString> ArtifactConfigurationSet;
 	TArray<TSharedPtr<FString>> ArtifactConfigurationList;
diff --git a/Engine/Source/Editor/StorageServerWidgets/Private/SBuildArtifactContentsViewer.cpp b/Engine/Source/Editor/StorageServerWidgets/Private/SBuildArtifactContentsViewer.cpp
new file mode 100644
--- /dev/null
+++ b/Engine/Source/Editor/StorageServerWidgets/Private/SBuildArtifactContentsViewer.cpp
@@ -0,0 +1,491 @@
+// Copyright Epic Games, Inc. All Rights Reserved.
+
+#include "SBuildArtifactContentsViewer.h"
+#include "Widgets/Text/STextBlock.h"
+#include "Widgets/Layout/SBorder.h"
+#include "Widgets/Layout/SSeparator.h"
+#include "Widgets/Input/SSearchBox.h"
+#include "Styling/StyleColors.h"
+#include "Styling/AppStyle.h"
+
+#define LOCTEXT_NAMESPACE "BuildArtifactContents"
+
+namespace UE::BuildArtifactContents::Internal
+{
+	const FName ColName = TEXT("Name");
+	const FName ColSize = TEXT("Size");
+}
+
+// ---------------------------------------------------------------------------
+// SContentsEntryRow
+// ---------------------------------------------------------------------------
+
+void SContentsEntryRow::Construct(const FArguments& InArgs, const TSharedRef<STableViewBase>& InOwnerTableView, TSharedPtr<FContentsEntry> InEntry)
+{
+	Entry = InEntry;
+	SMultiColumnTableRow<TSharedPtr<FContentsEntry>>::Construct(FSuperRowType::FArguments(), InOwnerTableView);
+}
+
+TSharedRef<SWidget> SContentsEntryRow::GenerateWidgetForColumn(const FName& ColumnName)
+{
+	using namespace UE::BuildArtifactContents::Internal;
+
+	if (Entry->bIsPartHeader)
+	{
+		if (ColumnName == ColName)
+		{
+			return SNew(STextBlock)
+				.Text(FText::FromString(Entry->Name))
+				.Font(FCoreStyle::GetDefaultFontStyle("Bold", 9))
+				.ColorAndOpacity(FStyleColors::AccentGreen)
+				.Margin(FMargin(4.0f, 4.0f));
+		}
+		return SNullWidget::NullWidget;
+	}
+
+	if (ColumnName == ColName)
+	{
+		return SNew(STextBlock)
+			.Text(FText::FromString(Entry->Name))
+			.Font(FCoreStyle::GetDefaultFontStyle("Regular", 9))
+			.Margin(FMargin(12.0f, 2.0f, 4.0f, 2.0f));
+	}
+	else if (ColumnName == ColSize)
+	{
+		return SNew(STextBlock)
+			.Text(FText::FromString(SBuildArtifactContentsViewer::FormatSize(Entry->RawSize)))
+			.Font(FCoreStyle::GetDefaultFontStyle("Regular", 9))
+			.Justification(ETextJustify::Right)
+			.Margin(FMargin(4.0f, 2.0f));
+	}
+	return SNullWidget::NullWidget;
+}
+
+// ---------------------------------------------------------------------------
+// SBuildArtifactContentsViewer
+// ---------------------------------------------------------------------------
+
+void SBuildArtifactContentsViewer::Construct(const FArguments& InArgs)
+{
+	using namespace UE::BuildArtifactContents::Internal;
+
+	ChildSlot
+	[
+		SNew(SVerticalBox)
+
+		// Header
+		+ SVerticalBox::Slot()
+		.AutoHeight()
+		.Padding(8.0f)
+		[
+			SNew(STextBlock)
+			.Text(FText::Format(LOCTEXT("ContentsHeader", "Contents for: {0}"), FText::FromString(InArgs._ArtifactName)))
+			.Font(FCoreStyle::GetDefaultFontStyle("Bold", 12))
+		]
+
+		+ SVerticalBox::Slot()
+		.AutoHeight()
+		[
+			SNew(SSeparator)
+		]
+
+		// Search box
+		+ SVerticalBox::Slot()
+		.AutoHeight()
+		.Padding(4.0f)
+		[
+			SNew(SSearchBox)
+			.HintText(LOCTEXT("SearchHint", "Filter files (wildcards: * ? ; separated)"))
+			.OnTextChanged(this, &SBuildArtifactContentsViewer::OnSearchTextChanged)
+		]
+
+		// Contents list
+		+ SVerticalBox::Slot()
+		.FillHeight(1.0f)
+		.Padding(4.0f)
+		[
+			SNew(SBorder)
+			.BorderImage(FAppStyle::GetBrush("ToolPanel.GroupBorder"))
+			.Padding(0.0f)
+			[
+				SAssignNew(ListView, SListView<TSharedPtr<FContentsEntry>>)
+				.ListItemsSource(&FilteredEntries)
+				.OnGenerateRow(this, &SBuildArtifactContentsViewer::OnGenerateRow)
+				.SelectionMode(ESelectionMode::Multi)
+				.HeaderRow
+				(
+					SNew(SHeaderRow)
+					+ SHeaderRow::Column(ColName)
+					.DefaultLabel(LOCTEXT("NameColumn", "Name"))
+					.FillWidth(0.85f)
+					.SortMode(this, &SBuildArtifactContentsViewer::GetColumnSortMode, ColName)
+					.OnSort(this, &SBuildArtifactContentsViewer::OnColumnSortModeChanged)
+
+					+ SHeaderRow::Column(ColSize)
+					.DefaultLabel(LOCTEXT("SizeColumn", "Size"))
+					.FillWidth(0.15f)
+					.HAlignCell(HAlign_Right)
+					.SortMode(this, &SBuildArtifactContentsViewer::GetColumnSortMode, ColSize)
+					.OnSort(this, &SBuildArtifactContentsViewer::OnColumnSortModeChanged)
+				)
+			]
+		]
+
+		// Status bar
+		+ SVerticalBox::Slot()
+		.AutoHeight()
+		.Padding(8.0f, 4.0f)
+		[
+			SAssignNew(StatusText, STextBlock)
+			.Text(LOCTEXT("Loading", "Loading..."))
+			.Font(FCoreStyle::GetDefaultFontStyle("Italic", 9))
+			.ColorAndOpacity(FStyleColors::Foreground.GetSpecifiedColor() * 0.6f)
+		]
+	];
+
+	RegisterActiveTimer(0.1f, FWidgetActiveTimerDelegate::CreateSP(this, &SBuildArtifactContentsViewer::PollPendingContents));
+}
+
+void SBuildArtifactContentsViewer::QueueContents(TMap<FString, UE::Zen::Build::FBuildServiceInstance::FBuildPart>&& BuildContents)
+{
+	FScopeLock Lock(&PendingContentsMutex);
+	PendingContents.Emplace(MoveTemp(BuildContents));
+}
+
+EActiveTimerReturnType SBuildArtifactContentsViewer::PollPendingContents(double InCurrentTime, float InDeltaTime)
+{
+	TOptional<TMap<FString, UE::Zen::Build::FBuildServiceInstance::FBuildPart>> Contents;
+	{
+		FScopeLock Lock(&PendingContentsMutex);
+		if (PendingContents.IsSet())
+		{
+			Contents.Emplace(MoveTemp(PendingContents.GetValue()));
+			PendingContents.Reset();
+		}
+	}
+
+	if (Contents.IsSet())
+	{
+		SetContents(Contents.GetValue());
+		return EActiveTimerReturnType::Stop;
+	}
+
+	return EActiveTimerReturnType::Continue;
+}
+
+void SBuildArtifactContentsViewer::SetContents(TMap<FString, UE::Zen::Build::FBuildServiceInstance::FBuildPart>& BuildContents)
+{
+	AllEntries.Empty();
+	TotalFileCount = 0;
+	TotalFileSize = 0;
+
+	// Sort parts by name
+	TArray<FString> PartNames;
+	BuildContents.GetKeys(PartNames);
+	PartNames.Sort();
+
+	for (const FString& PartName : PartNames)
+	{
+		UE::Zen::Build::FBuildServiceInstance::FBuildPart& Part = BuildContents[PartName];
+
+		// Add part header entry
+		TSharedPtr<FContentsEntry> HeaderEntry = MakeShared<FContentsEntry>();
+		HeaderEntry->bIsPartHeader = true;
+		HeaderEntry->Name = PartName;
+		AllEntries.Add(HeaderEntry);
+
+		// Sort files by name
+		Part.Files.Sort([](const UE::Zen::Build::FBuildServiceInstance::FBuildFile& A, const UE::Zen::Build::FBuildServiceInstance::FBuildFile& B)
+		{
+			return A.Name < B.Name;
+		});
+
+		for (const UE::Zen::Build::FBuildServiceInstance::FBuildFile& File : Part.Files)
+		{
+			TSharedPtr<FContentsEntry> FileEntry = MakeShared<FContentsEntry>();
+			FileEntry->bIsPartHeader = false;
+			FileEntry->Name = File.Name;
+			FileEntry->RawSize = File.RawSize;
+			AllEntries.Add(FileEntry);
+
+			TotalFileSize += File.RawSize;
+			TotalFileCount++;
+		}
+	}
+
+	bIsLoading = false;
+
+	SortAndRefreshList();
+}
+
+TSharedRef<ITableRow> SBuildArtifactContentsViewer::OnGenerateRow(TSharedPtr<FContentsEntry> Entry, const TSharedRef<STableViewBase>& OwnerTable)
+{
+	return SNew(SContentsEntryRow, OwnerTable, Entry);
+}
+
+FString SBuildArtifactContentsViewer::FormatSize(uint64 SizeInBytes)
+{
+	if (SizeInBytes < 1024)
+	{
+		return FString::Printf(TEXT("%llu B"), SizeInBytes);
+	}
+	if (SizeInBytes < 1024 * 1024)
+	{
+		return FString::Printf(TEXT("%.1f KB"), (double)SizeInBytes / 1024.0);
+	}
+	if (SizeInBytes < 1024ull * 1024 * 1024)
+	{
+		return FString::Printf(TEXT("%.1f MB"), (double)SizeInBytes / (1024.0 * 1024.0));
+	}
+	return FString::Printf(TEXT("%.2f GB"), (double)SizeInBytes / (1024.0 * 1024.0 * 1024.0));
+}
+
+void SBuildArtifactContentsViewer::SortAndRefreshList()
+{
+	using namespace UE::BuildArtifactContents::Internal;
+
+	if (SortMode != EColumnSortMode::None && !SortByColumn.IsNone())
+	{
+		// Separate part headers from file entries, sort files, then rebuild with headers
+		struct FPartGroup
+		{
+			TSharedPtr<FContentsEntry> Header;
+			TArray<TSharedPtr<FContentsEntry>> Files;
+		};
+
+		TArray<FPartGroup> Groups;
+		for (const TSharedPtr<FContentsEntry>& Entry : AllEntries)
+		{
+			if (Entry->bIsPartHeader)
+			{
+				Groups.AddDefaulted();
+				Groups.Last().Header = Entry;
+			}
+			else if (Groups.Num() > 0)
+			{
+				Groups.Last().Files.Add(Entry);
+			}
+		}
+
+		const bool bAscending = (SortMode == EColumnSortMode::Ascending);
+
+		for (FPartGroup& Group : Groups)
+		{
+			Group.Files.Sort([&](const TSharedPtr<FContentsEntry>& A, const TSharedPtr<FContentsEntry>& B)
+			{
+				if (SortByColumn == ColName)
+				{
+					return bAscending ? (A->Name < B->Name) : (A->Name > B->Name);
+				}
+				else if (SortByColumn == ColSize)
+				{
+					return bAscending ? (A->RawSize < B->RawSize) : (A->RawSize > B->RawSize);
+				}
+				return false;
+			});
+		}
+
+		AllEntries.Empty();
+		for (const FPartGroup& Group : Groups)
+		{
+			AllEntries.Add(Group.Header);
+			AllEntries.Append(Group.Files);
+		}
+	}
+
+	RefreshFilteredList();
+}
+
+void SBuildArtifactContentsViewer::RefreshFilteredList()
+{
+	FilteredEntries.Empty();
+
+	int32 FilteredFileCount = 0;
+	uint64 FilteredFileSize = 0;
+
+	if (SearchText.IsEmpty())
+	{
+		FilteredEntries = AllEntries;
+		FilteredFileCount = TotalFileCount;
+		FilteredFileSize = TotalFileSize;
+	}
+	else
+	{
+		TArray<FString> IncludeWildcards;
+		ParseIncludeWildcards(SearchText.ToString(), IncludeWildcards);
+
+		// Track which part headers should be included (only if they have matching files)
+		int32 PendingHeaderIndex = INDEX_NONE;
+		for (int32 i = 0; i < AllEntries.Num(); ++i)
+		{
+			const TSharedPtr<FContentsEntry>& Entry = AllEntries[i];
+			if (Entry->bIsPartHeader)
+			{
+				PendingHeaderIndex = i;
+			}
+			else
+			{
+				if (MatchesIncludeWildcards(IncludeWildcards, Entry->Name))
+				{
+					// Include the part header if we haven't already
+					if (PendingHeaderIndex != INDEX_NONE)
+					{
+						FilteredEntries.Add(AllEntries[PendingHeaderIndex]);
+						PendingHeaderIndex = INDEX_NONE;
+					}
+					FilteredEntries.Add(Entry);
+					FilteredFileCount++;
+					FilteredFileSize += Entry->RawSize;
+				}
+			}
+		}
+	}
+
+	if (ListView.IsValid())
+	{
+		ListView->RequestListRefresh();
+	}
+
+	if (StatusText.IsValid() && !bIsLoading)
+	{
+		if (!SearchText.IsEmpty() && FilteredFileCount != TotalFileCount)
+		{
+			StatusText->SetText(FText::Format(
+				LOCTEXT("StatusFormatFiltered", "Showing {0} of {1} {1}|plural(one=file,other=files) ({2} of {3})"),
+				FText::AsNumber(FilteredFileCount),
+				FText::AsNumber(TotalFileCount),
+				FText::FromString(FormatSize(FilteredFileSize)),
+				FText::FromString(FormatSize(TotalFileSize))
+			));
+		}
+		else
+		{
+			StatusText->SetText(FText::Format(
+				LOCTEXT("StatusFormat", "{0} {0}|plural(one=file,other=files), {1} total"),
+				FText::AsNumber(TotalFileCount),
+				FText::FromString(FormatSize(TotalFileSize))
+			));
+		}
+	}
+}
+
+void SBuildArtifactContentsViewer::OnSearchTextChanged(const FText& InText)
+{
+	SearchText = InText;
+	RefreshFilteredList();
+}
+
+EColumnSortMode::Type SBuildArtifactContentsViewer::GetColumnSortMode(const FName ColumnId) const
+{
+	if (SortByColumn == ColumnId)
+	{
+		return SortMode;
+	}
+	return EColumnSortMode::None;
+}
+
+void SBuildArtifactContentsViewer::OnColumnSortModeChanged(const EColumnSortPriority::Type SortPriority, const FName& ColumnId, const EColumnSortMode::Type InSortMode)
+{
+	SortByColumn = ColumnId;
+	SortMode = InSortMode;
+	SortAndRefreshList();
+}
+
+bool SBuildArtifactContentsViewer::MatchWildcard(FStringView Wildcard, FStringView String)
+{
+	const TCHAR* WildcardIt  = Wildcard.GetData();
+	const TCHAR* WildcardEnd = WildcardIt + Wildcard.Len();
+	const TCHAR* StringIt    = String.GetData();
+	const TCHAR* StringEnd   = StringIt + String.Len();
+
+	for (; WildcardIt != WildcardEnd; WildcardIt++)
+	{
+		TCHAR WC = TChar<TCHAR>::ToLower(*WildcardIt);
+		switch (WC)
+		{
+		case TEXT('?'):
+			if (StringIt == StringEnd)
+			{
+				return false;
+			}
+			StringIt++;
+			break;
+		case TEXT('*'):
+			if ((WildcardIt + 1) == WildcardEnd)
+			{
+				return true;
+			}
+			for (const TCHAR* It = StringIt; It <= StringEnd; It++)
+			{
+				if (MatchWildcard(FStringView(WildcardIt + 1, UE_PTRDIFF_TO_INT32(WildcardEnd - (WildcardIt + 1))),
+				                  FStringView(It, UE_PTRDIFF_TO_INT32(StringEnd - It))))
+				{
+					return true;
+				}
+			}
+			return false;
+		default:
+			if (StringIt == StringEnd)
+			{
+				return false;
+			}
+			if (TChar<TCHAR>::ToLower(*StringIt) != WC)
+			{
+				return false;
+			}
+			++StringIt;
+			break;
+		}
+	}
+	return StringIt == StringEnd;
+}
+
+void SBuildArtifactContentsViewer::ParseIncludeWildcards(const FString& FilterText, TArray<FString>& OutWildcards)
+{
+	TArray<FString> Tokens;
+	FilterText.ParseIntoArray(Tokens, TEXT(";"), true);
+
+	for (FString& Token : Tokens)
+	{
+		Token.TrimStartInline();
+		Token.TrimEndInline();
+		if (Token.IsEmpty())
+		{
+			continue;
+		}
+
+		// Lowercase and normalize separators
+		Token.ToLowerInline();
+		Token.ReplaceCharInline(TEXT('\\'), TEXT('/'));
+
+		// Strip leading ./
+		if (Token.StartsWith(TEXT("./")))
+		{
+			Token.RightChopInline(2);
+		}
+
+		if (!Token.IsEmpty())
+		{
+			OutWildcards.Add(MoveTemp(Token));
+		}
+	}
+}
+
+bool SBuildArtifactContentsViewer::MatchesIncludeWildcards(const TArray<FString>& IncludeWildcards, FStringView FileName)
+{
+	if (IncludeWildcards.IsEmpty())
+	{
+		return true;
+	}
+
+	for (const FString& Wildcard : IncludeWildcards)
+	{
+		if (MatchWildcard(Wildcard, FileName))
+		{
+			return true;
+		}
+	}
+	return false;
+}
+
+#undef LOCTEXT_NAMESPACE
diff --git a/Engine/Source/Editor/StorageServerWidgets/Private/SBuildArtifactContentsViewer.h b/Engine/Source/Editor/StorageServerWidgets/Private/SBuildArtifactContentsViewer.h
new file mode 100644
--- /dev/null
+++ b/Engine/Source/Editor/StorageServerWidgets/Private/SBuildArtifactContentsViewer.h
@@ -0,0 +1,95 @@
+// Copyright Epic Games, Inc. All Rights Reserved.
+
+#pragma once
+
+#include "CoreMinimal.h"
+#include "HAL/CriticalSection.h"
+#include "Widgets/DeclarativeSyntaxSupport.h"
+#include "Widgets/SCompoundWidget.h"
+#include "Widgets/Views/SListView.h"
+#include "Widgets/Views/STableRow.h"
+#include "Experimental/BuildServerInterface.h"
+
+struct FContentsEntry
+{
+	bool bIsPartHeader = false;
+	FString Name;       // Part name for headers, file name for files
+	uint64 RawSize = 0; // File entries only
+};
+
+/**
+ * Row widget for a single contents entry, using SMultiColumnTableRow so columns
+ * track header resizing automatically.
+ */
+class SContentsEntryRow : public SMultiColumnTableRow<TSharedPtr<FContentsEntry>>
+{
+public:
+	SLATE_BEGIN_ARGS(SContentsEntryRow) {}
+	SLATE_END_ARGS()
+
+	void Construct(const FArguments& InArgs, const TSharedRef<STableViewBase>& InOwnerTableView, TSharedPtr<FContentsEntry> InEntry);
+	virtual TSharedRef<SWidget> GenerateWidgetForColumn(const FName& ColumnName) override;
+
+private:
+	TSharedPtr<FContentsEntry> Entry;
+};
+
+/**
+ * Widget for viewing build artifact file contents in a searchable list grouped by build part
+ */
+class SBuildArtifactContentsViewer : public SCompoundWidget
+{
+public:
+	SLATE_BEGIN_ARGS(SBuildArtifactContentsViewer)
+	{}
+		SLATE_ARGUMENT(FString, ArtifactName)
+	SLATE_END_ARGS()
+
+	void Construct(const FArguments& InArgs);
+
+	/** Thread-safe: queues contents from a background thread for consumption on the Slate thread. */
+	void QueueContents(TMap<FString, UE::Zen::Build::FBuildServiceInstance::FBuildPart>&& BuildContents);
+
+	static FString FormatSize(uint64 SizeInBytes);
+
+private:
+	void SetContents(TMap<FString, UE::Zen::Build::FBuildServiceInstance::FBuildPart>& BuildContents);
+
+	EActiveTimerReturnType PollPendingContents(double InCurrentTime, float InDeltaTime);
+
+	TSharedRef<ITableRow> OnGenerateRow(TSharedPtr<FContentsEntry> Entry, const TSharedRef<STableViewBase>& OwnerTable);
+
+	void RefreshFilteredList();
+
+	void SortAndRefreshList();
+
+	void OnSearchTextChanged(const FText& InText);
+
+	EColumnSortMode::Type GetColumnSortMode(const FName ColumnId) const;
+	void OnColumnSortModeChanged(const EColumnSortPriority::Type SortPriority, const FName& ColumnId, const EColumnSortMode::Type InSortMode);
+
+	static bool MatchWildcard(FStringView Wildcard, FStringView String);
+	static void ParseIncludeWildcards(const FString& FilterText, TArray<FString>& OutWildcards);
+	static bool MatchesIncludeWildcards(const TArray<FString>& IncludeWildcards, FStringView FileName);
+
+	FCriticalSection PendingContentsMutex;
+	TOptional<TMap<FString, UE::Zen::Build::FBuildServiceInstance::FBuildPart>> PendingContents;
+
+	TArray<TSharedPtr<FContentsEntry>> AllEntries;
+
+	TArray<TSharedPtr<FContentsEntry>> FilteredEntries;
+
+	TSharedPtr<SListView<TSharedPtr<FContentsEntry>>> ListView;
+
+	TSharedPtr<STextBlock> StatusText;
+
+	FText SearchText;
+
+	FName SortByColumn;
+	EColumnSortMode::Type SortMode = EColumnSortMode::None;
+
+	int32 TotalFileCount = 0;
+	uint64 TotalFileSize = 0;
+
+	bool bIsLoading = true;
+};
diff --git a/Engine/Source/Editor/StorageServerWidgets/Private/SBuildArtifactSelection.cpp b/Engine/Source/Editor/StorageServerWidgets/Private/SBuildArtifactSelection.cpp
--- a/Engine/Source/Editor/StorageServerWidgets/Private/SBuildArtifactSelection.cpp
+++ b/Engine/Source/Editor/StorageServerWidgets/Private/SBuildArtifactSelection.cpp
@@ -12,6 +12,9 @@
 #include "Widgets/Docking/SDockTab.h"
 #include "ZenServiceInstanceManager.h"
 #include "SBuildArtifactMetadataViewer.h"
+#include "SBuildArtifactContentsViewer.h"
+#include "Dialog/DialogCommands.h"
+#include "Dialog/SCustomDialog.h"
 
 
 #define LOCTEXT_NAMESPACE "StorageServerBuild"
@@ -28,6 +31,8 @@ namespace UE::BuildArtifactSelection::Internal
 void SBuildArtifactSelection::Construct(const FArguments& InArgs)
 {
 	using namespace UE::BuildSelection::Internal;
+	BuildServiceInstance = InArgs._BuildServiceInstance;
+
 	const FSlateColor TitleColor = FStyleColors::AccentWhite;
 	const FSlateFontInfo TitleFont = FCoreStyle::GetDefaultFontStyle("Bold", 10);
 
@@ -121,10 +126,12 @@ void SBuildArtifactSelection::SetAvailableConfigurations(
 	}
 }
 
-void SBuildArtifactSelection::Refresh(const FStringView InGroupDisplayName, const TMap<FString, UE::BuildSelection::Internal::FArtifact>* InNamedArtifacts, FCbObjectId InHighlightBuildId)
+void SBuildArtifactSelection::Refresh(const FStringView InNamespace, const FStringView InGroupDisplayName, const TMap<FString, UE::BuildSelection::Internal::FArtifact>* InNamedArtifacts, FCbObjectId InHighlightBuildId, bool bInCanViewContents)
 {
 	using namespace UE::BuildSelection::Internal;
 
+	Namespace = InNamespace;
+	bCanViewContents = bInCanViewContents;
 	HighlightedBuildId = InHighlightBuildId;
 	GroupDisplayName = InGroupDisplayName;
 	if (InNamedArtifacts)
@@ -500,6 +507,69 @@ TSharedRef<SWidget> SBuildArtifactSelection::CreateArtifactMenuWidget(UE::BuildS
 		EUserInterfaceActionType::Button
 	);
 
+	if (bCanViewContents)
+	{
+	MenuBuilder.AddMenuEntry(
+		LOCTEXT("BuildSelection_ViewContents", "View Contents"),
+		FText(),
+		FSlateIcon(),
+		FUIAction(
+			FExecuteAction::CreateLambda([this, Artifact, ArtifactName]
+				{
+					TSharedPtr<UE::Zen::Build::FBuildServiceInstance> ServiceInstance = BuildServiceInstance.Get();
+					if (!ServiceInstance || !Artifact)
+					{
+						return;
+					}
+
+					FString FullDisplayName = ArtifactName;
+					if (!GroupDisplayName.IsEmpty())
+					{
+						FullDisplayName = FString::Printf(TEXT("%s - %s"), *GroupDisplayName, *ArtifactName);
+					}
+
+					TSharedRef<SBuildArtifactContentsViewer> Viewer =
+						SNew(SBuildArtifactContentsViewer)
+						.ArtifactName(FullDisplayName);
+
+					FString BuildIdString = FString(WriteToString<64>(Artifact->BuildRecord.BuildId));
+					TWeakPtr<SBuildArtifactContentsViewer> WeakViewer = Viewer;
+					ServiceInstance->ListBuildContents(Namespace, Artifact->BuildRecord.BucketId, BuildIdString,
+						[WeakViewer](TMap<FString, UE::Zen::Build::FBuildServiceInstance::FBuildPart>&& BuildContents)
+						{
+							if (TSharedPtr<SBuildArtifactContentsViewer> Viewer = WeakViewer.Pin())
+							{
+								Viewer->QueueContents(MoveTemp(BuildContents));
+							}
+						});
+
+					if (!FDialogCommands::IsRegistered())
+					{
+						FDialogCommands::Register();
+					}
+
+					TSharedRef<SCustomDialog> Dialog = SNew(SCustomDialog)
+						.Title(FText::Format(LOCTEXT("ViewContentsTitle", "Build Contents: {0}"), FText::FromString(FullDisplayName)))
+						.ClientSize(FVector2D(800, 600))
+						.UseScrollBox(false)
+						.HAlignContent(HAlign_Fill)
+						.VAlignContent(VAlign_Fill)
+						.Content()
+						[
+							Viewer
+						]
+						.Buttons({
+							SCustomDialog::FButton(LOCTEXT("Close", "Close"))
+						});
+
+					Dialog->ShowModal();
+				})
+		),
+		NAME_None,
+		EUserInterfaceActionType::Button
+	);
+	}
+
 	return MenuBuilder.MakeWidget();
 }
 
diff --git a/Engine/Source/Editor/StorageServerWidgets/Private/SBuildSelection.cpp b/Engine/Source/Editor/StorageServerWidgets/Private/SBuildSelection.cpp
--- a/Engine/Source/Editor/StorageServerWidgets/Private/SBuildSelection.cpp
+++ b/Engine/Source/Editor/StorageServerWidgets/Private/SBuildSelection.cpp
@@ -1092,7 +1092,7 @@ void SBuildSelection::RefreshArtifactSelectionWidget(FBuildSelectionBuildGroupPt
 			}
 		}
 	}
-	BuildArtifactSelection->Refresh(Item ? Item->DisplayName : FStringView(), Item ? &Item->NamedArtifacts : nullptr, BuildIdToHighlight);
+	BuildArtifactSelection->Refresh(Item ? FStringView(Item->Namespace) : FStringView(), Item ? Item->DisplayName : FStringView(), Item ? &Item->NamedArtifacts : nullptr, BuildIdToHighlight, GetSelectedBuildType() != EBuildType::Oplog);
 }
 
 void SBuildSelection::BuildGroupSelectionChanged(FBuildSelectionBuildGroupPtr Item, ESelectInfo::Type SelectInfo)
@@ -2272,6 +2272,7 @@ TSharedRef<SWidget> SBuildSelection::GetGridPanel()
 			.FillSize(1.0f)
 			[
 				SAssignNew(BuildArtifactSelection, SBuildArtifactSelection)
+					.BuildServiceInstance(BuildServiceInstance)
 			]
 		]
 		+SHorizontalBox::Slot()
