BuildStorageTool: Adjusting the position and padding of the context button for build artifacts.

#Synced-CL 51529686

diff --git a/Engine/Source/Editor/StorageServerWidgets/Private/SBuildArtifactSelection.cpp b/Engine/Source/Editor/StorageServerWidgets/Private/SBuildArtifactSelection.cpp
--- a/Engine/Source/Editor/StorageServerWidgets/Private/SBuildArtifactSelection.cpp
+++ b/Engine/Source/Editor/StorageServerWidgets/Private/SBuildArtifactSelection.cpp
@@ -366,8 +366,10 @@ void SBuildArtifactSelection::RebuildSelectedGroupArtifactContents()
 					]
 				]
 				+SHorizontalBox::Slot()
-				.AutoWidth()
+				.FillWidth(1.0f)
+				.HAlign(HAlign_Right)
 				.VAlign(VAlign_Center)
+				.Padding(0, 0, 8, 0)
 				[
 					SNew(SComboButton)
 					.ComboButtonStyle(&FAppStyle::Get().GetWidgetStyle<FComboButtonStyle>("Artifact.SmallContext"))
