BuildStorageTool: Add individual file downloading functionality for non-oplog builds.

#Synced-CL 51529686

diff --git a/Engine/Source/Developer/Zen/Private/BuildServerInterface.cpp b/Engine/Source/Developer/Zen/Private/BuildServerInterface.cpp
--- a/Engine/Source/Developer/Zen/Private/BuildServerInterface.cpp
+++ b/Engine/Source/Developer/Zen/Private/BuildServerInterface.cpp
@@ -148,7 +148,7 @@ FBuildTransferState::GetLogFilename() const
 	const FString LogDirectory = IFileManager::Get().ConvertToAbsolutePathForExternalAppForWrite(
 		*FPaths::GetPath(FPlatformOutputDevices::GetAbsoluteLogFilename())
 	);
-	return FPaths::Combine(LogDirectory, TEXT("Transfers"), Name) + TEXT(".log");
+	return FPaths::Combine(LogDirectory, TEXT("Transfers"), FPaths::MakeValidFileName(Name)) + TEXT(".log");
 }
 
 void FAccessToken::SetToken(const FStringView Scheme, const FStringView Token)
diff --git a/Engine/Source/Editor/StorageServerWidgets/Internal/SBuildArtifactSelection.h b/Engine/Source/Editor/StorageServerWidgets/Internal/SBuildArtifactSelection.h
--- a/Engine/Source/Editor/StorageServerWidgets/Internal/SBuildArtifactSelection.h
+++ b/Engine/Source/Editor/StorageServerWidgets/Internal/SBuildArtifactSelection.h
@@ -18,6 +18,8 @@
 
 #define UE_API STORAGESERVERWIDGETS_API
 
+DECLARE_DELEGATE_ThreeParams(FOnDownloadWithSpec, const FString& /*ArtifactName*/, const UE::BuildSelection::Internal::FArtifact& /*Artifact*/, const FString& /*DownloadSpecJSON*/);
+
 class SBuildArtifactSelection : public SCompoundWidget
 {
 public:
@@ -30,6 +32,8 @@ public:
 	SLATE_ATTRIBUTE(TSharedPtr<UE::Zen::FZenServiceInstance>, ZenServiceInstance);
 	SLATE_ATTRIBUTE(TSharedPtr<UE::Zen::Build::FBuildServiceInstance>, BuildServiceInstance);
 
+	SLATE_EVENT(FOnDownloadWithSpec, OnDownloadWithSpec)
+
 	SLATE_END_ARGS()
 
 	UE_API void Construct(const FArguments& InArgs);
@@ -49,6 +53,7 @@ private:
 	TSharedRef<SWidget> CreateArtifactMenuWidget(UE::BuildSelection::Internal::FArtifact* Artifact, const FString& ArtifactName);
 
 	TAttribute<TSharedPtr<UE::Zen::Build::FBuildServiceInstance>> BuildServiceInstance;
+	FOnDownloadWithSpec OnDownloadWithSpec;
 	FString Namespace;
 	bool bCanViewContents = false;
 
diff --git a/Engine/Source/Editor/StorageServerWidgets/Internal/SBuildSelection.h b/Engine/Source/Editor/StorageServerWidgets/Internal/SBuildSelection.h
--- a/Engine/Source/Editor/StorageServerWidgets/Internal/SBuildSelection.h
+++ b/Engine/Source/Editor/StorageServerWidgets/Internal/SBuildSelection.h
@@ -144,7 +144,7 @@ private:
 
 	FReply ExploreDestination_OnClicked();
 	void DownloadSelectedArtifactsInSelectedBuildGroup();
-	bool DownloadItem(const FBuildGroup& Group, const FString& ArtifactName, const UE::BuildSelection::Internal::FArtifact& Artifact);
+	bool DownloadItem(const FBuildGroup& Group, const FString& ArtifactName, const UE::BuildSelection::Internal::FArtifact& Artifact, FString&& DownloadSpecJSONContents = FString());
 
 	TMap<FString, FString> EngineInstallations;
 
diff --git a/Engine/Source/Editor/StorageServerWidgets/Private/SBuildArtifactContentsViewer.cpp b/Engine/Source/Editor/StorageServerWidgets/Private/SBuildArtifactContentsViewer.cpp
--- a/Engine/Source/Editor/StorageServerWidgets/Private/SBuildArtifactContentsViewer.cpp
+++ b/Engine/Source/Editor/StorageServerWidgets/Private/SBuildArtifactContentsViewer.cpp
@@ -1,9 +1,11 @@
 // Copyright Epic Games, Inc. All Rights Reserved.
 
 #include "SBuildArtifactContentsViewer.h"
+#include "Serialization/JsonWriter.h"
 #include "Widgets/Text/STextBlock.h"
 #include "Widgets/Layout/SBorder.h"
 #include "Widgets/Layout/SSeparator.h"
+#include "Widgets/Input/SButton.h"
 #include "Widgets/Input/SSearchBox.h"
 #include "Styling/StyleColors.h"
 #include "Styling/AppStyle.h"
@@ -48,6 +50,10 @@ TSharedRef<SWidget> SContentsEntryRow::GenerateWidgetForColumn(const FName& Colu
 		return SNew(STextBlock)
 			.Text(FText::FromString(Entry->Name))
 			.Font(FCoreStyle::GetDefaultFontStyle("Regular", 9))
+			.ColorAndOpacity_Lambda([this]()
+			{
+				return Entry->bIsMarked ? FSlateColor(FStyleColors::AccentYellow) : FSlateColor::UseForeground();
+			})
 			.Margin(FMargin(12.0f, 2.0f, 4.0f, 2.0f));
 	}
 	else if (ColumnName == ColSize)
@@ -55,6 +61,10 @@ TSharedRef<SWidget> SContentsEntryRow::GenerateWidgetForColumn(const FName& Colu
 		return SNew(STextBlock)
 			.Text(FText::FromString(SBuildArtifactContentsViewer::FormatSize(Entry->RawSize)))
 			.Font(FCoreStyle::GetDefaultFontStyle("Regular", 9))
+			.ColorAndOpacity_Lambda([this]()
+			{
+				return Entry->bIsMarked ? FSlateColor(FStyleColors::AccentYellow) : FSlateColor::UseForeground();
+			})
 			.Justification(ETextJustify::Right)
 			.Margin(FMargin(4.0f, 2.0f));
 	}
@@ -89,14 +99,65 @@ void SBuildArtifactContentsViewer::Construct(const FArguments& InArgs)
 			SNew(SSeparator)
 		]
 
-		// Search box
+		// Search box + Mark button
 		+ SVerticalBox::Slot()
 		.AutoHeight()
 		.Padding(4.0f)
 		[
-			SNew(SSearchBox)
-			.HintText(LOCTEXT("SearchHint", "Filter files (wildcards: * ? ; separated)"))
-			.OnTextChanged(this, &SBuildArtifactContentsViewer::OnSearchTextChanged)
+			SNew(SHorizontalBox)
+			+ SHorizontalBox::Slot()
+			.FillWidth(1.0f)
+			[
+				SNew(SSearchBox)
+				.HintText(LOCTEXT("SearchHint", "Filter files (wildcards: * ? ; separated)"))
+				.OnTextChanged(this, &SBuildArtifactContentsViewer::OnSearchTextChanged)
+			]
+			+ SHorizontalBox::Slot()
+			.AutoWidth()
+			.Padding(4.0f, 0.0f, 0.0f, 0.0f)
+			.VAlign(VAlign_Center)
+			[
+				SNew(SButton)
+				.Text(LOCTEXT("ToggleDownload", "Toggle Download"))
+				.IsEnabled_Lambda([this]()
+				{
+					if (!ListView.IsValid())
+					{
+						return false;
+					}
+					TArray<TSharedPtr<FContentsEntry>> SelectedItems = ListView->GetSelectedItems();
+					for (const TSharedPtr<FContentsEntry>& Item : SelectedItems)
+					{
+						if (!Item->bIsPartHeader)
+						{
+							return true;
+						}
+					}
+					return false;
+				})
+				.OnClicked_Lambda([this]()
+				{
+					MarkSelectedFiles();
+					return FReply::Handled();
+				})
+			]
+			+ SHorizontalBox::Slot()
+			.AutoWidth()
+			.Padding(4.0f, 0.0f, 0.0f, 0.0f)
+			.VAlign(VAlign_Center)
+			[
+				SNew(SButton)
+				.Text(LOCTEXT("ClearAll", "Clear All"))
+				.IsEnabled_Lambda([this]()
+				{
+					return MarkedFileCount > 0;
+				})
+				.OnClicked_Lambda([this]()
+				{
+					ClearAllMarkedFiles();
+					return FReply::Handled();
+				})
+			]
 		]
 
 		// Contents list
@@ -141,6 +202,20 @@ void SBuildArtifactContentsViewer::Construct(const FArguments& InArgs)
 			.Font(FCoreStyle::GetDefaultFontStyle("Italic", 9))
 			.ColorAndOpacity(FStyleColors::Foreground.GetSpecifiedColor() * 0.6f)
 		]
+
+		// Marked files status
+		+ SVerticalBox::Slot()
+		.AutoHeight()
+		.Padding(8.0f, 0.0f, 8.0f, 4.0f)
+		[
+			SAssignNew(MarkedStatusText, STextBlock)
+			.Font(FCoreStyle::GetDefaultFontStyle("Italic", 9))
+			.ColorAndOpacity(FStyleColors::AccentYellow)
+			.Visibility_Lambda([this]()
+			{
+				return MarkedFileCount > 0 ? EVisibility::Visible : EVisibility::Collapsed;
+			})
+		]
 	];
 
 	RegisterActiveTimer(0.1f, FWidgetActiveTimerDelegate::CreateSP(this, &SBuildArtifactContentsViewer::PollPendingContents));
@@ -192,6 +267,7 @@ void SBuildArtifactContentsViewer::SetContents(TMap<FString, UE::Zen::Build::FBu
 		TSharedPtr<FContentsEntry> HeaderEntry = MakeShared<FContentsEntry>();
 		HeaderEntry->bIsPartHeader = true;
 		HeaderEntry->Name = PartName;
+		HeaderEntry->PartId = Part.Id;
 		AllEntries.Add(HeaderEntry);
 
 		// Sort files by name
@@ -367,6 +443,15 @@ void SBuildArtifactContentsViewer::RefreshFilteredList()
 			));
 		}
 	}
+
+	if (MarkedStatusText.IsValid() && MarkedFileCount > 0)
+	{
+		MarkedStatusText->SetText(FText::Format(
+			LOCTEXT("MarkedStatusFormat", "{0} {0}|plural(one=file,other=files) marked for download ({1})"),
+			FText::AsNumber(MarkedFileCount),
+			FText::FromString(FormatSize(MarkedFileSize))
+		));
+	}
 }
 
 void SBuildArtifactContentsViewer::OnSearchTextChanged(const FText& InText)
@@ -488,4 +573,141 @@ bool SBuildArtifactContentsViewer::MatchesIncludeWildcards(const TArray<FString>
 	return false;
 }
 
+void SBuildArtifactContentsViewer::MarkSelectedFiles()
+{
+	if (!ListView.IsValid())
+	{
+		return;
+	}
+
+	TArray<TSharedPtr<FContentsEntry>> SelectedItems = ListView->GetSelectedItems();
+	for (const TSharedPtr<FContentsEntry>& Item : SelectedItems)
+	{
+		if (!Item->bIsPartHeader)
+		{
+			Item->bIsMarked = !Item->bIsMarked;
+		}
+	}
+
+	// Recompute marked counts from AllEntries
+	MarkedFileCount = 0;
+	MarkedFileSize = 0;
+	for (const TSharedPtr<FContentsEntry>& Entry : AllEntries)
+	{
+		if (!Entry->bIsPartHeader && Entry->bIsMarked)
+		{
+			MarkedFileCount++;
+			MarkedFileSize += Entry->RawSize;
+		}
+	}
+
+	if (MarkedStatusText.IsValid())
+	{
+		if (MarkedFileCount > 0)
+		{
+			MarkedStatusText->SetText(FText::Format(
+				LOCTEXT("MarkedStatusFormat", "{0} {0}|plural(one=file,other=files) marked for download ({1})"),
+				FText::AsNumber(MarkedFileCount),
+				FText::FromString(FormatSize(MarkedFileSize))
+			));
+		}
+	}
+
+	ListView->RequestListRefresh();
+}
+
+void SBuildArtifactContentsViewer::ClearAllMarkedFiles()
+{
+	for (const TSharedPtr<FContentsEntry>& Entry : AllEntries)
+	{
+		Entry->bIsMarked = false;
+	}
+
+	MarkedFileCount = 0;
+	MarkedFileSize = 0;
+
+	if (ListView.IsValid())
+	{
+		ListView->RequestListRefresh();
+	}
+}
+
+FReply SBuildArtifactContentsViewer::OnKeyDown(const FGeometry& MyGeometry, const FKeyEvent& InKeyEvent)
+{
+	if (InKeyEvent.GetKey() == EKeys::SpaceBar)
+	{
+		MarkSelectedFiles();
+		return FReply::Handled();
+	}
+	return SCompoundWidget::OnKeyDown(MyGeometry, InKeyEvent);
+}
+
+FString SBuildArtifactContentsViewer::ComposeDownloadSpecJSON() const
+{
+	// Group marked files by their parent part
+	struct FPartInfo
+	{
+		FCbObjectId PartId;
+		TArray<FString> Files;
+	};
+
+	TArray<TPair<FString, FPartInfo>> Parts;
+	FString CurrentPartName;
+	FCbObjectId CurrentPartId;
+
+	for (const TSharedPtr<FContentsEntry>& Entry : AllEntries)
+	{
+		if (Entry->bIsPartHeader)
+		{
+			CurrentPartName = Entry->Name;
+			CurrentPartId = Entry->PartId;
+		}
+		else if (Entry->bIsMarked)
+		{
+			// Find or create the part group
+			FPartInfo* Found = nullptr;
+			for (TPair<FString, FPartInfo>& Pair : Parts)
+			{
+				if (Pair.Key == CurrentPartName)
+				{
+					Found = &Pair.Value;
+					break;
+				}
+			}
+			if (!Found)
+			{
+				Parts.Emplace(CurrentPartName, FPartInfo{ CurrentPartId, {} });
+				Found = &Parts.Last().Value;
+			}
+			Found->Files.Add(Entry->Name);
+		}
+	}
+
+	FString OutputString;
+	TSharedRef<TJsonWriter<>> Writer = TJsonWriterFactory<>::Create(&OutputString);
+
+	Writer->WriteObjectStart();
+	Writer->WriteObjectStart(TEXT("parts"));
+	for (const TPair<FString, FPartInfo>& Pair : Parts)
+	{
+		Writer->WriteObjectStart(Pair.Key);
+		if (!(Pair.Value.PartId == FCbObjectId()))
+		{
+			Writer->WriteValue(TEXT("partId"), FString(WriteToString<64>(Pair.Value.PartId)));
+		}
+		Writer->WriteArrayStart(TEXT("files"));
+		for (const FString& FileName : Pair.Value.Files)
+		{
+			Writer->WriteValue(FileName);
+		}
+		Writer->WriteArrayEnd();
+		Writer->WriteObjectEnd();
+	}
+	Writer->WriteObjectEnd();
+	Writer->WriteObjectEnd();
+
+	Writer->Close();
+	return OutputString;
+}
+
 #undef LOCTEXT_NAMESPACE
diff --git a/Engine/Source/Editor/StorageServerWidgets/Private/SBuildArtifactContentsViewer.h b/Engine/Source/Editor/StorageServerWidgets/Private/SBuildArtifactContentsViewer.h
--- a/Engine/Source/Editor/StorageServerWidgets/Private/SBuildArtifactContentsViewer.h
+++ b/Engine/Source/Editor/StorageServerWidgets/Private/SBuildArtifactContentsViewer.h
@@ -13,8 +13,10 @@
 struct FContentsEntry
 {
 	bool bIsPartHeader = false;
-	FString Name;       // Part name for headers, file name for files
-	uint64 RawSize = 0; // File entries only
+	bool bIsMarked = false;     // Persistent download mark (file entries only)
+	FString Name;               // Part name for headers, file name for files
+	FCbObjectId PartId;         // Part ID (header entries only, for download spec)
+	uint64 RawSize = 0;         // File entries only
 };
 
 /**
@@ -50,6 +52,21 @@ public:
 	/** Thread-safe: queues contents from a background thread for consumption on the Slate thread. */
 	void QueueContents(TMap<FString, UE::Zen::Build::FBuildServiceInstance::FBuildPart>&& BuildContents);
 
+	/** Toggle bIsMarked on selected non-header items and update counts. */
+	void MarkSelectedFiles();
+
+	/** Unmark all marked files. */
+	void ClearAllMarkedFiles();
+
+	bool HasMarkedFiles() const { return MarkedFileCount > 0; }
+	int32 GetMarkedFileCount() const { return MarkedFileCount; }
+	uint64 GetMarkedFileSize() const { return MarkedFileSize; }
+
+	/** Build download spec JSON from marked entries. */
+	FString ComposeDownloadSpecJSON() const;
+
+	virtual FReply OnKeyDown(const FGeometry& MyGeometry, const FKeyEvent& InKeyEvent) override;
+
 	static FString FormatSize(uint64 SizeInBytes);
 
 private:
@@ -82,6 +99,7 @@ private:
 	TSharedPtr<SListView<TSharedPtr<FContentsEntry>>> ListView;
 
 	TSharedPtr<STextBlock> StatusText;
+	TSharedPtr<STextBlock> MarkedStatusText;
 
 	FText SearchText;
 
@@ -91,5 +109,8 @@ private:
 	int32 TotalFileCount = 0;
 	uint64 TotalFileSize = 0;
 
+	int32 MarkedFileCount = 0;
+	uint64 MarkedFileSize = 0;
+
 	bool bIsLoading = true;
 };
diff --git a/Engine/Source/Editor/StorageServerWidgets/Private/SBuildArtifactSelection.cpp b/Engine/Source/Editor/StorageServerWidgets/Private/SBuildArtifactSelection.cpp
--- a/Engine/Source/Editor/StorageServerWidgets/Private/SBuildArtifactSelection.cpp
+++ b/Engine/Source/Editor/StorageServerWidgets/Private/SBuildArtifactSelection.cpp
@@ -32,6 +32,7 @@ void SBuildArtifactSelection::Construct(const FArguments& InArgs)
 {
 	using namespace UE::BuildSelection::Internal;
 	BuildServiceInstance = InArgs._BuildServiceInstance;
+	OnDownloadWithSpec = InArgs._OnDownloadWithSpec;
 
 	const FSlateColor TitleColor = FStyleColors::AccentWhite;
 	const FSlateFontInfo TitleFont = FCoreStyle::GetDefaultFontStyle("Bold", 10);
@@ -550,6 +551,7 @@ TSharedRef<SWidget> SBuildArtifactSelection::CreateArtifactMenuWidget(UE::BuildS
 						FDialogCommands::Register();
 					}
 
+					TSharedPtr<SBuildArtifactContentsViewer> ViewerPtr = Viewer;
 					TSharedRef<SCustomDialog> Dialog = SNew(SCustomDialog)
 						.Title(FText::Format(LOCTEXT("ViewContentsTitle", "Build Contents: {0}"), FText::FromString(FullDisplayName)))
 						.ClientSize(FVector2D(800, 600))
@@ -561,10 +563,20 @@ TSharedRef<SWidget> SBuildArtifactSelection::CreateArtifactMenuWidget(UE::BuildS
 							Viewer
 						]
 						.Buttons({
+							SCustomDialog::FButton(LOCTEXT("DownloadMarked", "Download Marked"))
+								.SetIsEnabled(TAttribute<bool>::CreateLambda([ViewerPtr]()
+								{
+									return ViewerPtr.IsValid() && ViewerPtr->HasMarkedFiles();
+								})),
 							SCustomDialog::FButton(LOCTEXT("Close", "Close"))
 						});
 
-					Dialog->ShowModal();
+					int32 ButtonIndex = Dialog->ShowModal();
+					if (ButtonIndex == 0 && ViewerPtr->HasMarkedFiles())
+					{
+						FString DownloadSpec = ViewerPtr->ComposeDownloadSpecJSON();
+						OnDownloadWithSpec.ExecuteIfBound(ArtifactName, *Artifact, DownloadSpec);
+					}
 				})
 		),
 		NAME_None,
diff --git a/Engine/Source/Editor/StorageServerWidgets/Private/SBuildSelection.cpp b/Engine/Source/Editor/StorageServerWidgets/Private/SBuildSelection.cpp
--- a/Engine/Source/Editor/StorageServerWidgets/Private/SBuildSelection.cpp
+++ b/Engine/Source/Editor/StorageServerWidgets/Private/SBuildSelection.cpp
@@ -1496,7 +1496,7 @@ void SBuildSelection::DownloadSelectedArtifactsInSelectedBuildGroup()
 	}
 }
 
-bool SBuildSelection::DownloadItem(const FBuildGroup& Group, const FString& ArtifactName, const UE::BuildSelection::Internal::FArtifact& Artifact)
+bool SBuildSelection::DownloadItem(const FBuildGroup& Group, const FString& ArtifactName, const UE::BuildSelection::Internal::FArtifact& Artifact, FString&& DownloadSpecJSONContents)
 {
 	using namespace UE::Zen::Build;
 
@@ -1593,14 +1593,15 @@ bool SBuildSelection::DownloadItem(const FBuildGroup& Group, const FString& Arti
 				DestinationFolder,
 				Group.Namespace,
 				Artifact.BuildRecord.BucketId,
-				UserSelectedTransferFilters[BuildTypeIndex].IncludeFilter,
-				UserSelectedTransferFilters[BuildTypeIndex].ExcludeFilter,
-				FString(),
+				DownloadSpecJSONContents.IsEmpty() ? UserSelectedTransferFilters[BuildTypeIndex].IncludeFilter : FString(),
+				DownloadSpecJSONContents.IsEmpty() ? UserSelectedTransferFilters[BuildTypeIndex].ExcludeFilter : FString(),
+				MoveTemp(DownloadSpecJSONContents),
 				TArray<FString>(),
 				TArray<FString>(),
 				HostOverride,
 				RequestFlags);
-		OnBuildTransferStarted.ExecuteIfBound(BuildTransfer, Group.DisplayName, ArtifactName);
+		FString BuildTransferName = FString::Printf(TEXT("%s%s"), *Group.DisplayName, DownloadSpecJSONContents.IsEmpty() ? TEXT("") : TEXT("*"));
+		OnBuildTransferStarted.ExecuteIfBound(BuildTransfer, BuildTransferName, ArtifactName);
 	}
 
 	return true;
@@ -2273,6 +2274,14 @@ TSharedRef<SWidget> SBuildSelection::GetGridPanel()
 			[
 				SAssignNew(BuildArtifactSelection, SBuildArtifactSelection)
 					.BuildServiceInstance(BuildServiceInstance)
+					.OnDownloadWithSpec_Lambda([this](const FString& ArtifactName, const UE::BuildSelection::Internal::FArtifact& Artifact, const FString& Spec)
+					{
+						TArray<FBuildSelectionBuildGroupPtr> SelectedItems = BuildGroupListView->GetSelectedItems();
+						if (SelectedItems.Num() > 0)
+						{
+							DownloadItem(*SelectedItems[0], ArtifactName, Artifact, FString(Spec));
+						}
+					})
 			]
 		]
 		+SHorizontalBox::Slot()
