Build Storage Tool - Add build id to artifact metadata viewer

#Synced-CL 51823794

diff --git a/Engine/Source/Editor/StorageServerWidgets/Private/SBuildArtifactMetadataViewer.cpp b/Engine/Source/Editor/StorageServerWidgets/Private/SBuildArtifactMetadataViewer.cpp
--- a/Engine/Source/Editor/StorageServerWidgets/Private/SBuildArtifactMetadataViewer.cpp
+++ b/Engine/Source/Editor/StorageServerWidgets/Private/SBuildArtifactMetadataViewer.cpp
@@ -20,7 +20,7 @@ void SBuildArtifactMetadataViewer::Construct(const FArguments& InArgs)
 	ArtifactName = InArgs._ArtifactName;
 
 	// Parse the metadata
-	ParseMetadata(InArgs._Metadata);
+	ParseMetadata(InArgs._BuildId, InArgs._Metadata);
 
 	// Initialize filtered list with all entries
 	RefreshFilteredList();
@@ -133,10 +133,13 @@ TSharedRef<ITableRow> SBuildArtifactMetadataViewer::OnGenerateRow(TSharedPtr<FMe
 		];
 }
 
-void SBuildArtifactMetadataViewer::ParseMetadata(const FCbObject& Metadata)
+void SBuildArtifactMetadataViewer::ParseMetadata(const FCbObjectId& BuildId, const FCbObject& Metadata)
 {
 	MetadataEntries.Empty();
 
+	// Add the build id as the first entry
+	MetadataEntries.Add(MakeShared<FMetadataEntry>(TEXT("buildId"), FString(WriteToString<64>(BuildId)), TEXT("ObjectId")));
+
 	// Iterate through all fields in the compact binary object
 	for (FCbFieldView Field : Metadata)
 	{
diff --git a/Engine/Source/Editor/StorageServerWidgets/Private/SBuildArtifactMetadataViewer.h b/Engine/Source/Editor/StorageServerWidgets/Private/SBuildArtifactMetadataViewer.h
--- a/Engine/Source/Editor/StorageServerWidgets/Private/SBuildArtifactMetadataViewer.h
+++ b/Engine/Source/Editor/StorageServerWidgets/Private/SBuildArtifactMetadataViewer.h
@@ -30,6 +30,7 @@ public:
 	{}
 		SLATE_ARGUMENT(FCbObject, Metadata)
 		SLATE_ARGUMENT(FString, ArtifactName)
+		SLATE_ARGUMENT(FCbObjectId, BuildId)
 	SLATE_END_ARGS()
 
 	void Construct(const FArguments& InArgs);
@@ -42,7 +43,7 @@ private:
 	TSharedPtr<SWidget> OnContextMenuOpening();
 
 	/** Parse FCbObject metadata into display entries */
-	void ParseMetadata(const FCbObject& Metadata);
+	void ParseMetadata(const FCbObjectId& BuildId, const FCbObject& Metadata);
 
 	/** Refresh the filtered list based on search text */
 	void RefreshFilteredList();
diff --git a/Engine/Source/Editor/StorageServerWidgets/Private/SBuildArtifactSelection.cpp b/Engine/Source/Editor/StorageServerWidgets/Private/SBuildArtifactSelection.cpp
--- a/Engine/Source/Editor/StorageServerWidgets/Private/SBuildArtifactSelection.cpp
+++ b/Engine/Source/Editor/StorageServerWidgets/Private/SBuildArtifactSelection.cpp
@@ -532,6 +532,7 @@ TSharedRef<SWidget> SBuildArtifactSelection::CreateArtifactMenuWidget(UE::BuildS
 								SNew(SBuildArtifactMetadataViewer)
 								.Metadata(Artifact->BuildRecord.Metadata)
 								.ArtifactName(FullDisplayName)
+								.BuildId(Artifact->BuildRecord.BuildId)
 							];
 
 						// Find any existing live metadata tab
