BuildStorageTool: Change the contents viewer to use checkboxes for marking.  Fix appicon style. Address issues when trying to higlight a URL that is for a build that is filtered out of view.

#Synced-CL 51529686

diff --git a/Engine/Source/Editor/StorageServerWidgets/Private/SBuildArtifactContentsViewer.cpp b/Engine/Source/Editor/StorageServerWidgets/Private/SBuildArtifactContentsViewer.cpp
--- a/Engine/Source/Editor/StorageServerWidgets/Private/SBuildArtifactContentsViewer.cpp
+++ b/Engine/Source/Editor/StorageServerWidgets/Private/SBuildArtifactContentsViewer.cpp
@@ -7,6 +7,7 @@
 #include "Widgets/Layout/SBorder.h"
 #include "Widgets/Layout/SSeparator.h"
 #include "Widgets/Input/SButton.h"
+#include "Widgets/Input/SCheckBox.h"
 #include "Widgets/Input/SSearchBox.h"
 #include "Styling/StyleColors.h"
 #include "Styling/AppStyle.h"
@@ -15,6 +16,7 @@
 
 namespace UE::BuildArtifactContents::Internal
 {
+	const FName ColMark = TEXT("Mark");
 	const FName ColName = TEXT("Name");
 	const FName ColSize = TEXT("Size");
 }
@@ -26,6 +28,7 @@ namespace UE::BuildArtifactContents::Internal
 void SContentsEntryRow::Construct(const FArguments& InArgs, const TSharedRef<STableViewBase>& InOwnerTableView, TSharedPtr<FContentsEntry> InEntry)
 {
 	Entry = InEntry;
+	OnMarkChanged = InArgs._OnMarkChanged;
 	SMultiColumnTableRow<TSharedPtr<FContentsEntry>>::Construct(FSuperRowType::FArguments(), InOwnerTableView);
 }
 
@@ -46,26 +49,31 @@ TSharedRef<SWidget> SContentsEntryRow::GenerateWidgetForColumn(const FName& Colu
 		return SNullWidget::NullWidget;
 	}
 
-	if (ColumnName == ColName)
+	if (ColumnName == ColMark)
+	{
+		return SNew(SCheckBox)
+			.IsChecked_Lambda([this]()
+			{
+				return Entry->bIsMarked ? ECheckBoxState::Checked : ECheckBoxState::Unchecked;
+			})
+			.OnCheckStateChanged_Lambda([this](ECheckBoxState NewState)
+			{
+				Entry->bIsMarked = (NewState == ECheckBoxState::Checked);
+				OnMarkChanged.ExecuteIfBound();
+			});
+	}
+	else if (ColumnName == ColName)
 	{
 		return SNew(STextBlock)
 			.Text(FText::FromString(Entry->Name))
 			.Font(FCoreStyle::GetDefaultFontStyle("Regular", 9))
-			.ColorAndOpacity_Lambda([this]()
-			{
-				return Entry->bIsMarked ? FSlateColor(FStyleColors::AccentYellow) : FSlateColor::UseForeground();
-			})
-			.Margin(FMargin(12.0f, 2.0f, 4.0f, 2.0f));
+			.Margin(FMargin(4.0f, 2.0f));
 	}
 	else if (ColumnName == ColSize)
 	{
 		return SNew(STextBlock)
 			.Text(FText::FromString(SBuildArtifactContentsViewer::FormatSize(Entry->RawSize)))
 			.Font(FCoreStyle::GetDefaultFontStyle("Regular", 9))
-			.ColorAndOpacity_Lambda([this]()
-			{
-				return Entry->bIsMarked ? FSlateColor(FStyleColors::AccentYellow) : FSlateColor::UseForeground();
-			})
 			.Justification(ETextJustify::Right)
 			.Margin(FMargin(4.0f, 2.0f));
 	}
@@ -100,7 +108,7 @@ void SBuildArtifactContentsViewer::Construct(const FArguments& InArgs)
 			SNew(SSeparator)
 		]
 
-		// Search box + Mark button
+		// Search box + Mark All / Clear All buttons
 		+ SVerticalBox::Slot()
 		.AutoHeight()
 		.Padding(4.0f)
@@ -119,26 +127,14 @@ void SBuildArtifactContentsViewer::Construct(const FArguments& InArgs)
 			.VAlign(VAlign_Center)
 			[
 				SNew(SButton)
-				.Text(LOCTEXT("ToggleDownload", "Toggle Download"))
+				.Text(LOCTEXT("MarkAll", "Mark All"))
 				.IsEnabled_Lambda([this]()
 				{
-					if (!ListView.IsValid())
-					{
-						return false;
-					}
-					TArray<TSharedPtr<FContentsEntry>> SelectedItems = ListView->GetSelectedItems();
-					for (const TSharedPtr<FContentsEntry>& Item : SelectedItems)
-					{
-						if (!Item->bIsPartHeader)
-						{
-							return true;
-						}
-					}
-					return false;
+					return MarkedFileCount < TotalFileCount;
 				})
 				.OnClicked_Lambda([this]()
 				{
-					MarkSelectedFiles();
+					MarkAllFiles();
 					return FReply::Handled();
 				})
 			]
@@ -177,6 +173,12 @@ void SBuildArtifactContentsViewer::Construct(const FArguments& InArgs)
 				.HeaderRow
 				(
 					SNew(SHeaderRow)
+					+ SHeaderRow::Column(ColMark)
+					.DefaultLabel(FText::GetEmpty())
+					.FixedWidth(24.0f)
+					.HAlignHeader(HAlign_Center)
+					.HAlignCell(HAlign_Center)
+
 					+ SHeaderRow::Column(ColName)
 					.DefaultLabel(LOCTEXT("NameColumn", "Name"))
 					.FillWidth(0.85f)
@@ -297,7 +299,8 @@ void SBuildArtifactContentsViewer::SetContents(TMap<FString, UE::Zen::Build::FBu
 
 TSharedRef<ITableRow> SBuildArtifactContentsViewer::OnGenerateRow(TSharedPtr<FContentsEntry> Entry, const TSharedRef<STableViewBase>& OwnerTable)
 {
-	return SNew(SContentsEntryRow, OwnerTable, Entry);
+	return SNew(SContentsEntryRow, OwnerTable, Entry)
+		.OnMarkChanged(FSimpleDelegate::CreateSP(this, &SBuildArtifactContentsViewer::RecalculateMarkedCounts));
 }
 
 FString SBuildArtifactContentsViewer::FormatSize(uint64 SizeInBytes)
@@ -613,7 +616,11 @@ void SBuildArtifactContentsViewer::MarkSelectedFiles()
 		}
 	}
 
-	// Recompute marked counts from AllEntries
+	RecalculateMarkedCounts();
+}
+
+void SBuildArtifactContentsViewer::RecalculateMarkedCounts()
+{
 	MarkedFileCount = 0;
 	MarkedFileSize = 0;
 	for (const TSharedPtr<FContentsEntry>& Entry : AllEntries)
@@ -637,7 +644,23 @@ void SBuildArtifactContentsViewer::MarkSelectedFiles()
 		}
 	}
 
-	ListView->RequestListRefresh();
+	if (ListView.IsValid())
+	{
+		ListView->RequestListRefresh();
+	}
+}
+
+void SBuildArtifactContentsViewer::MarkAllFiles()
+{
+	for (const TSharedPtr<FContentsEntry>& Entry : AllEntries)
+	{
+		if (!Entry->bIsPartHeader)
+		{
+			Entry->bIsMarked = true;
+		}
+	}
+
+	RecalculateMarkedCounts();
 }
 
 void SBuildArtifactContentsViewer::ClearAllMarkedFiles()
diff --git a/Engine/Source/Editor/StorageServerWidgets/Private/SBuildArtifactContentsViewer.h b/Engine/Source/Editor/StorageServerWidgets/Private/SBuildArtifactContentsViewer.h
--- a/Engine/Source/Editor/StorageServerWidgets/Private/SBuildArtifactContentsViewer.h
+++ b/Engine/Source/Editor/StorageServerWidgets/Private/SBuildArtifactContentsViewer.h
@@ -27,6 +27,7 @@ class SContentsEntryRow : public SMultiColumnTableRow<TSharedPtr<FContentsEntry>
 {
 public:
 	SLATE_BEGIN_ARGS(SContentsEntryRow) {}
+		SLATE_EVENT(FSimpleDelegate, OnMarkChanged)
 	SLATE_END_ARGS()
 
 	void Construct(const FArguments& InArgs, const TSharedRef<STableViewBase>& InOwnerTableView, TSharedPtr<FContentsEntry> InEntry);
@@ -34,6 +35,7 @@ public:
 
 private:
 	TSharedPtr<FContentsEntry> Entry;
+	FSimpleDelegate OnMarkChanged;
 };
 
 /**
@@ -55,6 +57,9 @@ public:
 	/** Toggle bIsMarked on selected non-header items and update counts. */
 	void MarkSelectedFiles();
 
+	/** Mark all file entries for download. */
+	void MarkAllFiles();
+
 	/** Unmark all marked files. */
 	void ClearAllMarkedFiles();
 
@@ -62,6 +67,9 @@ public:
 	int32 GetMarkedFileCount() const { return MarkedFileCount; }
 	uint64 GetMarkedFileSize() const { return MarkedFileSize; }
 
+	/** Recompute MarkedFileCount/MarkedFileSize from AllEntries and update status text. */
+	void RecalculateMarkedCounts();
+
 	/** Build download spec JSON from marked entries. */
 	FString ComposeDownloadSpecJSON() const;
 
diff --git a/Engine/Source/Editor/StorageServerWidgets/Private/SBuildArtifactSelection.cpp b/Engine/Source/Editor/StorageServerWidgets/Private/SBuildArtifactSelection.cpp
--- a/Engine/Source/Editor/StorageServerWidgets/Private/SBuildArtifactSelection.cpp
+++ b/Engine/Source/Editor/StorageServerWidgets/Private/SBuildArtifactSelection.cpp
@@ -6,9 +6,11 @@
 #include "HAL/PlatformApplicationMisc.h"
 #include "Framework/Docking/TabManager.h"
 #include "Styling/StyleColors.h"
+#include "Framework/Application/SlateApplication.h"
 #include "Widgets/Input/SButton.h"
 #include "Widgets/Input/SCheckBox.h"
 #include "Widgets/Input/SComboBox.h"
+#include "Widgets/Layout/SBorder.h"
 #include "Widgets/Text/STextBlock.h"
 #include "Widgets/Docking/SDockTab.h"
 #include "ZenServiceInstanceManager.h"
@@ -320,83 +322,104 @@ void SBuildArtifactSelection::RebuildSelectedGroupArtifactContents()
 				.AutoWidth()
 				.VAlign(VAlign_Center)
 				[
-					SNew(SButton)
-					.ButtonStyle(FAppStyle::Get(), "InvisibleButton")
-					.IsFocusable(false)
-
-					.OnClicked_Lambda([this, CurrentCheckbox, ArtifactName]()
+					SNew(SBorder)
+					.BorderImage(FStyleDefaults::GetNoBrush())
+					.Padding(0)
+					.OnMouseButtonDown_Lambda([this, ArtifactName](const FGeometry&, const FPointerEvent& MouseEvent) -> FReply
+					{
+						if (MouseEvent.GetEffectingButton() == EKeys::RightMouseButton)
 						{
-							ECheckBoxState NewState = ECheckBoxState::Checked;
-							if (CurrentCheckbox->IsChecked())
-							{
-								NewState = ECheckBoxState::Unchecked;
-							}
-							CurrentCheckbox->SetIsChecked(NewState);
-							if (NewState == ECheckBoxState::Checked)
-							{
-								SelectedArtifacts.AddUnique(ArtifactName);
-							}
-							else if (NewState == ECheckBoxState::Unchecked)
-							{
-								SelectedArtifacts.Remove(ArtifactName);
-							}
+							TSharedRef<SWidget> MenuContent = CreateArtifactMenuWidget(NamedArtifacts.Find(ArtifactName), ArtifactName);
+							FSlateApplication::Get().PushMenu(
+								SharedThis(this),
+								FWidgetPath(),
+								MenuContent,
+								FSlateApplication::Get().GetCursorPos(),
+								FPopupTransitionEffect(FPopupTransitionEffect::ContextMenu)
+							);
 							return FReply::Handled();
-						})
+						}
+						return FReply::Unhandled();
+					})
 					[
-						SNew(STextBlock)
-						.Justification(ETextJustify::Left)
-						.Text(FText::FromString(ArtifactName))
-						.ToolTipText_Lambda([this, ArtifactName]() -> FText
-						{
-							if (const FArtifact* Artifact = NamedArtifacts.Find(ArtifactName))
+						SNew(SButton)
+						.ButtonStyle(FAppStyle::Get(), "InvisibleButton")
+						.IsFocusable(false)
+
+						.OnClicked_Lambda([this, CurrentCheckbox, ArtifactName]()
 							{
-								FCbFieldView TotalSizeField = Artifact->BuildRecord.Metadata["totalSize"];
-								if (TotalSizeField.HasValue() && !TotalSizeField.HasError())
+								ECheckBoxState NewState = ECheckBoxState::Checked;
+								if (CurrentCheckbox->IsChecked())
 								{
-									uint64 TotalSize = 0;
-									if (TotalSizeField.IsInteger())
-									{
-										TotalSize = TotalSizeField.AsUInt64();
-									}
-									else if (TotalSizeField.IsFloat())
-									{
-										TotalSize = static_cast<uint64>(TotalSizeField.AsDouble());
-									}
-									else if (TotalSizeField.IsString())
-									{
-										LexFromString(TotalSize, *FString(FUTF8ToTCHAR(TotalSizeField.AsString())));
-									}
-									else
+									NewState = ECheckBoxState::Unchecked;
+								}
+								CurrentCheckbox->SetIsChecked(NewState);
+								if (NewState == ECheckBoxState::Checked)
+								{
+									SelectedArtifacts.AddUnique(ArtifactName);
+								}
+								else if (NewState == ECheckBoxState::Unchecked)
+								{
+									SelectedArtifacts.Remove(ArtifactName);
+								}
+								return FReply::Handled();
+							})
+						[
+							SNew(STextBlock)
+							.Justification(ETextJustify::Left)
+							.Text(FText::FromString(ArtifactName))
+							.ToolTipText_Lambda([this, ArtifactName]() -> FText
+							{
+								if (const FArtifact* Artifact = NamedArtifacts.Find(ArtifactName))
+								{
+									FCbFieldView TotalSizeField = Artifact->BuildRecord.Metadata["totalSize"];
+									if (TotalSizeField.HasValue() && !TotalSizeField.HasError())
 									{
-										return FText::GetEmpty();
+										uint64 TotalSize = 0;
+										if (TotalSizeField.IsInteger())
+										{
+											TotalSize = TotalSizeField.AsUInt64();
+										}
+										else if (TotalSizeField.IsFloat())
+										{
+											TotalSize = static_cast<uint64>(TotalSizeField.AsDouble());
+										}
+										else if (TotalSizeField.IsString())
+										{
+											LexFromString(TotalSize, *FString(FUTF8ToTCHAR(TotalSizeField.AsString())));
+										}
+										else
+										{
+											return FText::GetEmpty();
+										}
+										return FText::FromString(SBuildArtifactContentsViewer::FormatSize(TotalSize));
 									}
-									return FText::FromString(SBuildArtifactContentsViewer::FormatSize(TotalSize));
 								}
-							}
-							return FText::GetEmpty();
-						})
-						.Font_Lambda([this, ArtifactName, RegularFont, HighlightFont]
-						{
-							if (FArtifact* Artifact = NamedArtifacts.Find(ArtifactName))
+								return FText::GetEmpty();
+							})
+							.Font_Lambda([this, ArtifactName, RegularFont, HighlightFont]
 							{
-								if (Artifact->BuildRecord.BuildId == HighlightedBuildId)
+								if (FArtifact* Artifact = NamedArtifacts.Find(ArtifactName))
 								{
-									return HighlightFont;
+									if (Artifact->BuildRecord.BuildId == HighlightedBuildId)
+									{
+										return HighlightFont;
+									}
 								}
-							}
-							return RegularFont;
-						})
-						.ColorAndOpacity_Lambda([this, ArtifactName]
-						{
-							if (FArtifact* Artifact = NamedArtifacts.Find(ArtifactName))
+								return RegularFont;
+							})
+							.ColorAndOpacity_Lambda([this, ArtifactName]
 							{
-								if (Artifact->BuildRecord.BuildId == HighlightedBuildId)
+								if (FArtifact* Artifact = NamedArtifacts.Find(ArtifactName))
 								{
-									return FSlateColor(FStyleColors::AccentYellow);
+									if (Artifact->BuildRecord.BuildId == HighlightedBuildId)
+									{
+										return FSlateColor(FStyleColors::AccentYellow);
+									}
 								}
-							}
-							return FSlateColor::UseForeground();
-						})
+								return FSlateColor::UseForeground();
+							})
+						]
 					]
 				]
 				+SHorizontalBox::Slot()
diff --git a/Engine/Source/Editor/StorageServerWidgets/Private/SBuildSelection.cpp b/Engine/Source/Editor/StorageServerWidgets/Private/SBuildSelection.cpp
--- a/Engine/Source/Editor/StorageServerWidgets/Private/SBuildSelection.cpp
+++ b/Engine/Source/Editor/StorageServerWidgets/Private/SBuildSelection.cpp
@@ -279,6 +279,7 @@ void SBuildSelection::ActOnBuildArtifact(EBuildArtifactAction Action, const UE::
 			[this]
 			{
 				EnforceTargetedBuildActionOnGroup();
+				RefreshPresentedBuildGroups();
 			});
 	}
 }
@@ -294,6 +295,7 @@ void SBuildSelection::EnforceTargetedBuildActionOnGroup()
 		Action = TargetedBuildIdAction;
 		BuildId = TargetedBuildId;
 	}
+
 	if (Action == EBuildArtifactAction::None)
 	{
 		return;
@@ -879,14 +881,36 @@ void SBuildSelection::RegenerateActivePlatformFilters()
 
 void SBuildSelection::ValidateBuildGroupSelection()
 {
+	using namespace UE::BuildSelection::Internal;
+
 	BuildGroupListView->UpdateSelectionSet();
 	TArray<FBuildSelectionBuildGroupPtr> SelectedItems = BuildGroupListView->GetSelectedItems();
 	if (SelectedItems.IsEmpty())
 	{
 		return;
 	}
+
+	EBuildArtifactAction Action;
+	FCbObjectId BuildId;
+	{
+		FReadScopeLock ReadLock(TargetedBuildIdLock);
+		Action = TargetedBuildIdAction;
+		BuildId = TargetedBuildId;
+	}
+
 	for (const FBuildSelectionBuildGroupPtr& SelectedItem : SelectedItems)
 	{
+		if (Action != EBuildArtifactAction::None)
+		{
+			// If a row contains a targeted build that we're actively trying to highlight, we won't force it to be de-selected due to filter criteria
+			for (const TPair<FString, FArtifact>& NamedArtifact : SelectedItem->NamedArtifacts)
+			{
+				if (NamedArtifact.Value.BuildRecord.BuildId == BuildId)
+				{
+					continue;
+				}
+			}
+		}
 		if (!BuildGroupIsSelectableOrNavigable(SelectedItem))
 		{
 			BuildGroupListView->SetItemSelection(SelectedItem, false);
@@ -990,8 +1014,8 @@ bool SBuildSelection::DoesPassAllViewFilters(TSharedPtr<FBuildGroup> InItem) con
 
 void SBuildSelection::PostViewFilterChange()
 {
-	ValidateBuildGroupSelection();
 	RefreshPresentedBuildGroups();
+	ValidateBuildGroupSelection();
 	BuildGroupListView->Invalidate(EInvalidateWidgetReason::LayoutAndVolatility);
 	BuildGroupListView->RequestListRefresh();
 }
@@ -1086,10 +1110,10 @@ void SBuildSelection::RefreshArtifactSelectionWidget(FBuildSelectionBuildGroupPt
 		{
 			FWriteScopeLock WriteLock(TargetedBuildIdLock);
 			Action = TargetedBuildIdAction;
+			TargetedBuildIdAction = EBuildArtifactAction::None;
 			if (Action != EBuildArtifactAction::None)
 			{
 				BuildIdToHighlight = TargetedBuildId;
-				TargetedBuildId = FCbObjectId();
 			}
 		}
 	}
@@ -1098,6 +1122,7 @@ void SBuildSelection::RefreshArtifactSelectionWidget(FBuildSelectionBuildGroupPt
 
 void SBuildSelection::BuildGroupSelectionChanged(FBuildSelectionBuildGroupPtr Item, ESelectInfo::Type SelectInfo)
 {
+	RefreshPresentedBuildGroups();
 	RefreshArtifactSelectionWidget(Item);
 }
 
diff --git a/Engine/Source/Programs/BuildStorageTool/Private/BuildStorageToolStyle.cpp b/Engine/Source/Programs/BuildStorageTool/Private/BuildStorageToolStyle.cpp
--- a/Engine/Source/Programs/BuildStorageTool/Private/BuildStorageToolStyle.cpp
+++ b/Engine/Source/Programs/BuildStorageTool/Private/BuildStorageToolStyle.cpp
@@ -59,7 +59,13 @@ TSharedRef< FSlateStyleSet > FBuildStorageToolStyle::Create()
 	const FVector2D Icon12x12(12.0f, 12.0f);
 	const FVector2D Icon16x16(16.0f, 16.0f);
 
-	Style.Set( "AppIcon", new IMAGE_BRUSH( "Icons/DefaultAppIcon", FVector2D(20,20)) );
+	{
+		const FString IconPath = FPaths::EngineSourceDir() / TEXT("Programs/BuildStorageTool/Resources/Windows/icon_24x24.png");
+		Style.Set("AppIcon", new FSlateImageBrush(IconPath, FVector2D(24, 24)));
+		Style.Set("AppIcon.Small", new FSlateImageBrush(IconPath, FVector2D(24, 24)));
+		Style.Set("AppIconPadding", FMargin(5.f, 5.f, 5.f, 5.f));
+		Style.Set("AppIconPadding.Small", FMargin(4.f, 4.f, 4.f, 4.f));
+	}
 
 	Style.Set("Zen.FolderExplore", new EDITOR_IMAGE_BRUSH_SVG("Starship/Common/ContentBrowser", Icon12x12));
 	Style.Set("Zen.BrowseContent", new EDITOR_IMAGE_BRUSH_SVG("Starship/Common/BrowseContent", Icon16x16));
