BuildStorageTool: Adding the ability to open individual files in the build contents viewer.

#Synced-CL 51529686

diff --git a/Engine/Source/Developer/Zen/Private/BuildServerInterface.cpp b/Engine/Source/Developer/Zen/Private/BuildServerInterface.cpp
--- a/Engine/Source/Developer/Zen/Private/BuildServerInterface.cpp
+++ b/Engine/Source/Developer/Zen/Private/BuildServerInterface.cpp
@@ -1316,6 +1316,11 @@ FBuildServiceInstance::BuildTransferThreadLoop()
 				{
 					CommandLineBuilder.Append(TEXT(" --verify"));
 				}
+
+				if (EnumHasAnyFlags(State->RequestFlags, EBuildTransferRequestFlags::Append))
+				{
+					CommandLineBuilder.Append(TEXT(" --append"));
+				}
 			break;
 		case FBuildTransfer::EType::Oplog:
 			{
diff --git a/Engine/Source/Developer/Zen/Public/Experimental/BuildServerInterface.h b/Engine/Source/Developer/Zen/Public/Experimental/BuildServerInterface.h
--- a/Engine/Source/Developer/Zen/Public/Experimental/BuildServerInterface.h
+++ b/Engine/Source/Developer/Zen/Public/Experimental/BuildServerInterface.h
@@ -59,6 +59,7 @@ enum class EBuildTransferRequestFlags
 	BoostWorkerCount = 1 << 4,
 	BoostWorkerMemory = 1 << 5,
 	BoostWorkers = BoostWorkerCount | BoostWorkerMemory,
+	Append = 1 << 6,
 };
 ENUM_CLASS_FLAGS(EBuildTransferRequestFlags);
 
diff --git a/Engine/Source/Editor/StorageServerWidgets/Private/SBuildArtifactContentsViewer.cpp b/Engine/Source/Editor/StorageServerWidgets/Private/SBuildArtifactContentsViewer.cpp
--- a/Engine/Source/Editor/StorageServerWidgets/Private/SBuildArtifactContentsViewer.cpp
+++ b/Engine/Source/Editor/StorageServerWidgets/Private/SBuildArtifactContentsViewer.cpp
@@ -1,6 +1,7 @@
 // Copyright Epic Games, Inc. All Rights Reserved.
 
 #include "SBuildArtifactContentsViewer.h"
+#include "Framework/MultiBox/MultiBoxBuilder.h"
 #include "Internationalization/Internationalization.h"
 #include "Serialization/JsonWriter.h"
 #include "Widgets/Text/STextBlock.h"
@@ -11,6 +12,7 @@
 #include "Widgets/Input/SSearchBox.h"
 #include "Styling/StyleColors.h"
 #include "Styling/AppStyle.h"
+#include "Widgets/Images/SThrobber.h"
 
 #define LOCTEXT_NAMESPACE "BuildArtifactContents"
 
@@ -64,10 +66,29 @@ TSharedRef<SWidget> SContentsEntryRow::GenerateWidgetForColumn(const FName& Colu
 	}
 	else if (ColumnName == ColName)
 	{
-		return SNew(STextBlock)
-			.Text(FText::FromString(Entry->Name))
-			.Font(FCoreStyle::GetDefaultFontStyle("Regular", 9))
-			.Margin(FMargin(4.0f, 2.0f));
+		return SNew(SHorizontalBox)
+			+ SHorizontalBox::Slot()
+			.FillWidth(1.0f)
+			[
+				SNew(STextBlock)
+				.Text(FText::FromString(Entry->Name))
+				.Font(FCoreStyle::GetDefaultFontStyle("Regular", 9))
+				.Margin(FMargin(4.0f, 2.0f))
+			]
+			+ SHorizontalBox::Slot()
+			.AutoWidth()
+			.VAlign(VAlign_Center)
+			.Padding(4.0f, 0.0f)
+			[
+				SNew(SCircularThrobber)
+				.Radius(7.f)
+				.NumPieces(6)
+				.Period(0.75f)
+				.Visibility_Lambda([this]()
+				{
+					return Entry->bIsDownloading ? EVisibility::Visible : EVisibility::Collapsed;
+				})
+			];
 	}
 	else if (ColumnName == ColSize)
 	{
@@ -84,10 +105,20 @@ TSharedRef<SWidget> SContentsEntryRow::GenerateWidgetForColumn(const FName& Colu
 // SBuildArtifactContentsViewer
 // ---------------------------------------------------------------------------
 
+SBuildArtifactContentsViewer::~SBuildArtifactContentsViewer()
+{
+	for (FPendingFileOpen& Pending : PendingFileOpens)
+	{
+		Pending.Transfer.RequestCancel();
+	}
+}
+
 void SBuildArtifactContentsViewer::Construct(const FArguments& InArgs)
 {
 	using namespace UE::BuildArtifactContents::Internal;
 
+	OnStartFileDownload = InArgs._OnStartFileDownload;
+
 	ChildSlot
 	[
 		SNew(SVerticalBox)
@@ -169,6 +200,7 @@ void SBuildArtifactContentsViewer::Construct(const FArguments& InArgs)
 				SAssignNew(ListView, SListView<TSharedPtr<FContentsEntry>>)
 				.ListItemsSource(&FilteredEntries)
 				.OnGenerateRow(this, &SBuildArtifactContentsViewer::OnGenerateRow)
+				.OnContextMenuOpening(this, &SBuildArtifactContentsViewer::OnContextMenuOpening)
 				.SelectionMode(ESelectionMode::Multi)
 				.HeaderRow
 				(
@@ -265,11 +297,13 @@ void SBuildArtifactContentsViewer::SetContents(TMap<FString, UE::Zen::Build::FBu
 	for (const FString& PartName : PartNames)
 	{
 		UE::Zen::Build::FBuildServiceInstance::FBuildPart& Part = BuildContents[PartName];
+		FName PartFName(*PartName);
 
 		// Add part header entry
 		TSharedPtr<FContentsEntry> HeaderEntry = MakeShared<FContentsEntry>();
 		HeaderEntry->bIsPartHeader = true;
 		HeaderEntry->Name = PartName;
+		HeaderEntry->PartName = PartFName;
 		HeaderEntry->PartId = Part.Id;
 		AllEntries.Add(HeaderEntry);
 
@@ -284,6 +318,8 @@ void SBuildArtifactContentsViewer::SetContents(TMap<FString, UE::Zen::Build::FBu
 			TSharedPtr<FContentsEntry> FileEntry = MakeShared<FContentsEntry>();
 			FileEntry->bIsPartHeader = false;
 			FileEntry->Name = File.Name;
+			FileEntry->PartName = PartFName;
+			FileEntry->PartId = Part.Id;
 			FileEntry->RawSize = File.RawSize;
 			AllEntries.Add(FileEntry);
 
@@ -699,35 +735,174 @@ FString SBuildArtifactContentsViewer::ComposeDownloadSpecJSON() const
 	};
 
 	TArray<TPair<FString, FPartInfo>> Parts;
-	FString CurrentPartName;
-	FCbObjectId CurrentPartId;
 
 	for (const TSharedPtr<FContentsEntry>& Entry : AllEntries)
 	{
-		if (Entry->bIsPartHeader)
+		if (Entry->bIsPartHeader || !Entry->bIsMarked)
 		{
-			CurrentPartName = Entry->Name;
-			CurrentPartId = Entry->PartId;
+			continue;
 		}
-		else if (Entry->bIsMarked)
+
+		FString PartNameStr = Entry->PartName.ToString();
+
+		// Find or create the part group
+		FPartInfo* Found = nullptr;
+		for (TPair<FString, FPartInfo>& Pair : Parts)
 		{
-			// Find or create the part group
-			FPartInfo* Found = nullptr;
-			for (TPair<FString, FPartInfo>& Pair : Parts)
+			if (Pair.Key == PartNameStr)
 			{
-				if (Pair.Key == CurrentPartName)
-				{
-					Found = &Pair.Value;
-					break;
-				}
+				Found = &Pair.Value;
+				break;
 			}
-			if (!Found)
+		}
+		if (!Found)
+		{
+			Parts.Emplace(PartNameStr, FPartInfo{ Entry->PartId, {} });
+			Found = &Parts.Last().Value;
+		}
+		Found->Files.Add(Entry->Name);
+	}
+
+	FString OutputString;
+	TSharedRef<TJsonWriter<>> Writer = TJsonWriterFactory<>::Create(&OutputString);
+
+	Writer->WriteObjectStart();
+	Writer->WriteObjectStart(TEXT("parts"));
+	for (const TPair<FString, FPartInfo>& Pair : Parts)
+	{
+		Writer->WriteObjectStart(Pair.Key);
+		if (!(Pair.Value.PartId == FCbObjectId()))
+		{
+			Writer->WriteValue(TEXT("partId"), FString(WriteToString<64>(Pair.Value.PartId)));
+		}
+		Writer->WriteArrayStart(TEXT("files"));
+		for (const FString& FileName : Pair.Value.Files)
+		{
+			Writer->WriteValue(FileName);
+		}
+		Writer->WriteArrayEnd();
+		Writer->WriteObjectEnd();
+	}
+	Writer->WriteObjectEnd();
+	Writer->WriteObjectEnd();
+
+	Writer->Close();
+	return OutputString;
+}
+
+TSharedPtr<SWidget> SBuildArtifactContentsViewer::OnContextMenuOpening()
+{
+	TArray<TSharedPtr<FContentsEntry>> SelectedItems = ListView->GetSelectedItems();
+
+	// Filter to non-header file entries only
+	TArray<TSharedPtr<FContentsEntry>> FileEntries;
+	for (const TSharedPtr<FContentsEntry>& Item : SelectedItems)
+	{
+		if (!Item->bIsPartHeader)
+		{
+			FileEntries.Add(Item);
+		}
+	}
+	if (FileEntries.IsEmpty() || !OnStartFileDownload.IsBound())
+	{
+		return nullptr;
+	}
+
+	FMenuBuilder MenuBuilder(/*bCloseAfterSelection*/ true, nullptr, TSharedPtr<FExtender>(),
+		/*bCloseSelfOnly*/ false, &FCoreStyle::Get(), /*bSearchable*/ false);
+
+	FText Label = FileEntries.Num() == 1
+		? LOCTEXT("OpenFile", "Open")
+		: FText::Format(LOCTEXT("OpenFiles", "Open {0} Files"), FText::AsNumber(FileEntries.Num()));
+	FText Tooltip = FileEntries.Num() == 1
+		? LOCTEXT("OpenFileTooltip", "Download this file and open it with the default application")
+		: LOCTEXT("OpenFilesTooltip", "Download selected files and open them with the default application");
+
+	FText ExploreLabel = FileEntries.Num() == 1
+		? LOCTEXT("ExploreFile", "Open in File Browser")
+		: FText::Format(LOCTEXT("ExploreFiles", "Open {0} Files in File Browser"), FText::AsNumber(FileEntries.Num()));
+	FText ExploreTooltip = FileEntries.Num() == 1
+		? LOCTEXT("ExploreFileTooltip", "Download this file and open its folder in the file browser")
+		: LOCTEXT("ExploreFilesTooltip", "Download selected files and open their folders in the file browser");
+
+	// Capture a copy for the explore entry before moving FileEntries into the open entry
+	TArray<TSharedPtr<FContentsEntry>> ExploreEntries = FileEntries;
+
+	MenuBuilder.AddMenuEntry(
+		Label,
+		Tooltip,
+		FSlateIcon(),
+		FUIAction(FExecuteAction::CreateLambda([this, Entries = MoveTemp(FileEntries)]()
+		{
+			// Must copy because FUIAction stores a shared const copy of the delegate
+			TArray<TSharedPtr<FContentsEntry>> EntriesCopy = Entries;
+			OpenFileEntries(MoveTemp(EntriesCopy));
+		}))
+	);
+
+	MenuBuilder.AddMenuEntry(
+		ExploreLabel,
+		ExploreTooltip,
+		FSlateIcon(),
+		FUIAction(FExecuteAction::CreateLambda([this, Entries = MoveTemp(ExploreEntries)]()
+		{
+			TArray<TSharedPtr<FContentsEntry>> EntriesCopy = Entries;
+			OpenFileEntries(MoveTemp(EntriesCopy), EFileOpenAction::ExploreFolder);
+		}))
+	);
+
+	return MenuBuilder.MakeWidget();
+}
+
+void SBuildArtifactContentsViewer::OpenFileEntries(TArray<TSharedPtr<FContentsEntry>>&& FileEntries, EFileOpenAction Action)
+{
+	FString Spec = ComposeFileDownloadSpec(FileEntries);
+	UE::Zen::Build::FBuildServiceInstance::FBuildTransfer Transfer = OnStartFileDownload.Execute(MoveTemp(Spec));
+
+	for (const TSharedPtr<FContentsEntry>& Entry : FileEntries)
+	{
+		Entry->bIsDownloading = true;
+	}
+
+	bool bNeedTimer = PendingFileOpens.IsEmpty();
+	PendingFileOpens.Add({ MoveTemp(Transfer), MoveTemp(FileEntries), Action });
+
+	if (bNeedTimer)
+	{
+		RegisterActiveTimer(0.5f, FWidgetActiveTimerDelegate::CreateSP(this, &SBuildArtifactContentsViewer::PollPendingFileOpens));
+	}
+}
+
+FString SBuildArtifactContentsViewer::ComposeFileDownloadSpec(const TArray<TSharedPtr<FContentsEntry>>& FileEntries) const
+{
+	struct FPartInfo
+	{
+		FCbObjectId PartId;
+		TArray<FString> Files;
+	};
+
+	TArray<TPair<FString, FPartInfo>> Parts;
+
+	for (const TSharedPtr<FContentsEntry>& FileEntry : FileEntries)
+	{
+		FString PartNameStr = FileEntry->PartName.ToString();
+
+		// Find or create the part group
+		FPartInfo* Found = nullptr;
+		for (TPair<FString, FPartInfo>& Pair : Parts)
+		{
+			if (Pair.Key == PartNameStr)
 			{
-				Parts.Emplace(CurrentPartName, FPartInfo{ CurrentPartId, {} });
-				Found = &Parts.Last().Value;
+				Found = &Pair.Value;
+				break;
 			}
-			Found->Files.Add(Entry->Name);
 		}
+		if (!Found)
+		{
+			Parts.Emplace(PartNameStr, FPartInfo{ FileEntry->PartId, {} });
+			Found = &Parts.Last().Value;
+		}
+		Found->Files.Add(FileEntry->Name);
 	}
 
 	FString OutputString;
@@ -757,4 +932,58 @@ FString SBuildArtifactContentsViewer::ComposeDownloadSpecJSON() const
 	return OutputString;
 }
 
+EActiveTimerReturnType SBuildArtifactContentsViewer::PollPendingFileOpens(double InCurrentTime, float InDeltaTime)
+{
+	using namespace UE::Zen::Build;
+
+	for (int32 i = PendingFileOpens.Num() - 1; i >= 0; --i)
+	{
+		FPendingFileOpen& Pending = PendingFileOpens[i];
+		FBuildServiceInstance::EBuildTransferStatus Status = Pending.Transfer.GetStatus();
+
+		if (Status == FBuildServiceInstance::EBuildTransferStatus::Succeeded)
+		{
+			switch (Pending.Action)
+			{
+			case EFileOpenAction::LaunchFile:
+				for (const TSharedPtr<FContentsEntry>& Entry : Pending.Entries)
+				{
+					FString FullPath = FPaths::Combine(Pending.Transfer.GetDestination(), Entry->Name);
+					FPlatformProcess::LaunchFileInDefaultExternalApplication(*FullPath);
+					Entry->bIsDownloading = false;
+				}
+				break;
+
+			case EFileOpenAction::ExploreFolder:
+				{
+					TSet<FString> FoldersToExplore;
+					for (const TSharedPtr<FContentsEntry>& Entry : Pending.Entries)
+					{
+						FString FullPath = FPaths::Combine(Pending.Transfer.GetDestination(), Entry->Name);
+						FoldersToExplore.Add(FPaths::GetPath(FullPath));
+						Entry->bIsDownloading = false;
+					}
+					for (const FString& Folder : FoldersToExplore)
+					{
+						FPlatformProcess::ExploreFolder(*Folder);
+					}
+				}
+				break;
+			}
+			PendingFileOpens.RemoveAt(i);
+		}
+		else if (Status == FBuildServiceInstance::EBuildTransferStatus::Failed ||
+		         Status == FBuildServiceInstance::EBuildTransferStatus::Canceled)
+		{
+			for (const TSharedPtr<FContentsEntry>& Entry : Pending.Entries)
+			{
+				Entry->bIsDownloading = false;
+			}
+			PendingFileOpens.RemoveAt(i);
+		}
+	}
+
+	return PendingFileOpens.IsEmpty() ? EActiveTimerReturnType::Stop : EActiveTimerReturnType::Continue;
+}
+
 #undef LOCTEXT_NAMESPACE
diff --git a/Engine/Source/Editor/StorageServerWidgets/Private/SBuildArtifactContentsViewer.h b/Engine/Source/Editor/StorageServerWidgets/Private/SBuildArtifactContentsViewer.h
--- a/Engine/Source/Editor/StorageServerWidgets/Private/SBuildArtifactContentsViewer.h
+++ b/Engine/Source/Editor/StorageServerWidgets/Private/SBuildArtifactContentsViewer.h
@@ -14,8 +14,10 @@ struct FContentsEntry
 {
 	bool bIsPartHeader = false;
 	bool bIsMarked = false;     // Persistent download mark (file entries only)
+	bool bIsDownloading = false; // true while an "Open" download is in progress
 	FString Name;               // Part name for headers, file name for files
-	FCbObjectId PartId;         // Part ID (header entries only, for download spec)
+	FName PartName;             // Parent part name (set on all entries)
+	FCbObjectId PartId;         // Parent part ID (set on all entries)
 	uint64 RawSize = 0;         // File entries only
 };
 
@@ -38,6 +40,12 @@ private:
 	FSimpleDelegate OnMarkChanged;
 };
 
+DECLARE_DELEGATE_RetVal_OneParam(
+	UE::Zen::Build::FBuildServiceInstance::FBuildTransfer,  // return: transfer handle
+	FOnStartFileDownload,
+	FString&&  /* DownloadSpecJSON */
+);
+
 /**
  * Widget for viewing build artifact file contents in a searchable list grouped by build part
  */
@@ -47,9 +55,11 @@ public:
 	SLATE_BEGIN_ARGS(SBuildArtifactContentsViewer)
 	{}
 		SLATE_ARGUMENT(FString, ArtifactName)
+		SLATE_EVENT(FOnStartFileDownload, OnStartFileDownload)
 	SLATE_END_ARGS()
 
 	void Construct(const FArguments& InArgs);
+	~SBuildArtifactContentsViewer();
 
 	/** Thread-safe: queues contents from a background thread for consumption on the Slate thread. */
 	void QueueContents(TMap<FString, UE::Zen::Build::FBuildServiceInstance::FBuildPart>&& BuildContents);
@@ -82,6 +92,17 @@ private:
 
 	EActiveTimerReturnType PollPendingContents(double InCurrentTime, float InDeltaTime);
 
+	enum class EFileOpenAction : uint8
+	{
+		LaunchFile,     // Open each file in its default application
+		ExploreFolder,  // Open the containing folder(s) in the native file browser
+	};
+
+	TSharedPtr<SWidget> OnContextMenuOpening();
+	void OpenFileEntries(TArray<TSharedPtr<FContentsEntry>>&& FileEntries, EFileOpenAction Action = EFileOpenAction::LaunchFile);
+	FString ComposeFileDownloadSpec(const TArray<TSharedPtr<FContentsEntry>>& FileEntries) const;
+	EActiveTimerReturnType PollPendingFileOpens(double InCurrentTime, float InDeltaTime);
+
 	TSharedRef<ITableRow> OnGenerateRow(TSharedPtr<FContentsEntry> Entry, const TSharedRef<STableViewBase>& OwnerTable);
 
 	void RefreshFilteredList();
@@ -120,5 +141,15 @@ private:
 	int32 MarkedFileCount = 0;
 	uint64 MarkedFileSize = 0;
 
+	struct FPendingFileOpen
+	{
+		UE::Zen::Build::FBuildServiceInstance::FBuildTransfer Transfer;
+		TArray<TSharedPtr<FContentsEntry>> Entries;
+		EFileOpenAction Action = EFileOpenAction::LaunchFile;
+	};
+
+	TArray<FPendingFileOpen> PendingFileOpens;
+	FOnStartFileDownload OnStartFileDownload;
+
 	bool bIsLoading = true;
 };
diff --git a/Engine/Source/Editor/StorageServerWidgets/Private/SBuildArtifactSelection.cpp b/Engine/Source/Editor/StorageServerWidgets/Private/SBuildArtifactSelection.cpp
--- a/Engine/Source/Editor/StorageServerWidgets/Private/SBuildArtifactSelection.cpp
+++ b/Engine/Source/Editor/StorageServerWidgets/Private/SBuildArtifactSelection.cpp
@@ -3,6 +3,8 @@
 #include "SBuildArtifactSelection.h"
 
 #include "Framework/MultiBox/MultiBoxBuilder.h"
+#include "Misc/CoreDelegates.h"
+#include "HAL/FileManager.h"
 #include "HAL/PlatformApplicationMisc.h"
 #include "Framework/Docking/TabManager.h"
 #include "Styling/StyleColors.h"
@@ -20,6 +22,27 @@
 #include "Dialog/SCustomDialog.h"
 
 
+namespace
+{
+	FString GetOpenFileTempRoot()
+	{
+		static FString TempRoot = FPaths::Combine(
+			FPlatformProcess::UserTempDir(),
+			TEXT("BuildStorageTool"),
+			TEXT("Open"));
+		static bool bCleanupRegistered = false;
+		if (!bCleanupRegistered)
+		{
+			bCleanupRegistered = true;
+			FCoreDelegates::OnPreExit.AddStatic([]()
+			{
+				IFileManager::Get().DeleteDirectory(*GetOpenFileTempRoot(), false, true);
+			});
+		}
+		return TempRoot;
+	}
+}
+
 #define LOCTEXT_NAMESPACE "StorageServerBuild"
 
 namespace UE::BuildArtifactSelection::Internal
@@ -617,7 +640,29 @@ TSharedRef<SWidget> SBuildArtifactSelection::CreateArtifactMenuWidget(UE::BuildS
 
 					TSharedRef<SBuildArtifactContentsViewer> Viewer =
 						SNew(SBuildArtifactContentsViewer)
-						.ArtifactName(FullDisplayName);
+						.ArtifactName(FullDisplayName)
+						.OnStartFileDownload_Lambda(
+							[ServiceInstance, Artifact, ArtifactName,
+							 GroupDisplayName = this->GroupDisplayName, Namespace = this->Namespace]
+							(FString&& DownloadSpec) -> UE::Zen::Build::FBuildServiceInstance::FBuildTransfer
+						{
+							FString DestFolder = FPaths::Combine(
+								GetOpenFileTempRoot(), GroupDisplayName, ArtifactName);
+
+							return ServiceInstance->StartBuildTransfer(
+								Artifact->BuildRecord.BuildId,
+								TEXT("Open"),
+								DestFolder,
+								Namespace,
+								Artifact->BuildRecord.BucketId,
+								FString(),
+								FString(),
+								MoveTemp(DownloadSpec),
+								TArray<FString>(),
+								TArray<FString>(),
+								FStringView(),
+								UE::Zen::Build::EBuildTransferRequestFlags::Append);
+						});
 
 					FString BuildIdString = FString(WriteToString<64>(Artifact->BuildRecord.BuildId));
 					TWeakPtr<SBuildArtifactContentsViewer> WeakViewer = Viewer;
