BuildStorageTool: Handle without crashing the possibility that zero namespaces+buckets may be returned.

#Synced-CL 51529686

diff --git a/Engine/Source/Editor/StorageServerWidgets/Private/SBuildSelection.cpp b/Engine/Source/Editor/StorageServerWidgets/Private/SBuildSelection.cpp
--- a/Engine/Source/Editor/StorageServerWidgets/Private/SBuildSelection.cpp
+++ b/Engine/Source/Editor/StorageServerWidgets/Private/SBuildSelection.cpp
@@ -534,6 +534,15 @@ void SBuildSelection::RebuildLists()
 
 	if (TSharedPtr<FBuildServiceInstance> ServiceInstance = BuildServiceInstance.Get())
 	{
+		BuildGroups.Empty();
+		PresentedBuildGroups.Reset();
+
+		if (NamespacesAndBuckets.IsEmpty() || !SelectedProject || !SelectedBuildType || !SelectedStream)
+		{
+			BuildListRefreshesInProgress.fetch_sub(1);
+			return;
+		}
+
 		FString BucketRegex = FString::Printf(TEXT("%s\\.%s\\.%s\\..*"), **SelectedProject, **SelectedBuildType, **SelectedStream);
 		FString BucketPrefix = FString::Printf(TEXT("%s.%s.%s."), **SelectedProject, **SelectedBuildType, **SelectedStream);
 
@@ -559,8 +568,6 @@ void SBuildSelection::RebuildLists()
 		uint32 QueryIndex = 0;
 
 		++BuildRefreshGeneration;
-		BuildGroups.Empty();
-		PresentedBuildGroups.Reset();
 
 		for (const FString& Namespace : NamespacesToQuery)
 		{
