BuildStorageTool: UX and usability improvements.  Build activity uses a throbber instead of overlaid progress bar.  Using edit verb when opening files in the contents viewer.  Changing the pattern for launching file explorer from contents view to only open one explorer at the highest common path.  Adjusted the way window foreground and draw attention is used when we highlight a URL.

#Synced-CL 51529686

diff --git a/Engine/Source/Editor/StorageServerWidgets/Private/SBuildActivity.cpp b/Engine/Source/Editor/StorageServerWidgets/Private/SBuildActivity.cpp
--- a/Engine/Source/Editor/StorageServerWidgets/Private/SBuildActivity.cpp
+++ b/Engine/Source/Editor/StorageServerWidgets/Private/SBuildActivity.cpp
@@ -11,7 +11,7 @@
 #include "Widgets/Input/SButton.h"
 #include "Widgets/Input/SHyperlink.h"
 #include "Widgets/Layout/SGridPanel.h"
-#include "Widgets/Notifications/SProgressBar.h"
+#include "Widgets/Images/SThrobber.h"
 #include "Widgets/SBoxPanel.h"
 #include "Widgets/Text/STextBlock.h"
 #include "Widgets/Views/SListView.h"
@@ -95,11 +95,11 @@ TSharedRef<SWidget> SBuildActivity::GetGridPanel()
 				SNew(SHeaderRow)
 				.Visibility(EVisibility::Collapsed)
 				+ SHeaderRow::Column(FBuildActivityIds::ColStatus).DefaultLabel(LOCTEXT("BuildActivity_ActivityStatus", "Status"))
-					.FillWidth(0.1f)
+					.FillWidth(0.05f)
 				+ SHeaderRow::Column(FBuildActivityIds::ColName).DefaultLabel(LOCTEXT("BuildActivity_ActivityName", "Name"))
 					.FillWidth(0.2f).HAlignCell(HAlign_Left).HAlignHeader(HAlign_Left).VAlignCell(VAlign_Center)
 				+ SHeaderRow::Column(FBuildActivityIds::ColPlatform).DefaultLabel(LOCTEXT("BuildActivity_ActivityPlatform", "Platform"))
-					.FillWidth(0.1f).HAlignCell(HAlign_Left).HAlignHeader(HAlign_Left).VAlignCell(VAlign_Center)
+					.FillWidth(0.15f).HAlignCell(HAlign_Left).HAlignHeader(HAlign_Left).VAlignCell(VAlign_Center)
 				+ SHeaderRow::Column(FBuildActivityIds::ColDuration).DefaultLabel(LOCTEXT("BuildActivity_ActivityDuration", "Duration"))
 					.FillWidth(0.05f).HAlignCell(HAlign_Left).HAlignHeader(HAlign_Left).VAlignCell(VAlign_Center)
 				+ SHeaderRow::Column(FBuildActivityIds::ColDescription).DefaultLabel(LOCTEXT("BuildActivity_ActivityDescription", "Description"))
@@ -425,6 +425,8 @@ SBuildActivityTableRow::GenerateWidgetForColumn(const FName& ColumnName)
 	{
 		return SNew(SHorizontalBox)
 				+ SHorizontalBox::Slot()
+				.AutoWidth()
+				.VAlign(VAlign_Center)
 				.Padding(1.f,1.f)
 				[
 					SNew(STextBlock)
@@ -452,6 +454,19 @@ SBuildActivityTableRow::GenerateWidgetForColumn(const FName& ColumnName)
 							}
 							return LOCTEXT("BuildActivity_TransferStatusUnknown", "Unknown");
 						})
+				]
+				+ SHorizontalBox::Slot()
+				.AutoWidth()
+				.VAlign(VAlign_Center)
+				.Padding(4.f,1.f,1.f,1.f)
+				[
+					SNew(SCircularThrobber)
+					.Radius(6.f)
+					.Visibility_Lambda([this]()
+					{
+						return BuildActivity->Transfer.GetStatus() == FBuildServiceInstance::EBuildTransferStatus::Active
+							? EVisibility::Visible : EVisibility::Collapsed;
+					})
 				];
 	}
 	else if (ColumnName == FBuildActivityIds::ColName)
@@ -500,55 +515,35 @@ SBuildActivityTableRow::GenerateWidgetForColumn(const FName& ColumnName)
 	}
 	else if (ColumnName == FBuildActivityIds::ColDescription)
 	{
-		return SNew(SOverlay)
-					+ SOverlay::Slot()
-					.Padding(1.f,1.f)
-					[
-						SNew(SProgressBar)
-						.Visibility_Lambda([this]()
-						{
-							return BuildActivity->Transfer.GetStatus() == FBuildServiceInstance::EBuildTransferStatus::Active ? EVisibility::Visible : EVisibility::Collapsed;
-						})
-						.Percent_Lambda([this]()
-						{
-							FString Label;
-							FString Detail;
-							float Percent;
-							if (BuildActivity->Transfer.GetCurrentProgress(Label, Detail, Percent))
-							{
-								return Percent/100.f;
-							}
-							return 0.f;
-						})
-					]
-					+ SOverlay::Slot()
-					.Padding(1.f,1.f)
-					.VAlign(VAlign_Center)
-					[
-						SNew(STextBlock)
-						.Text_Lambda([this]()
-						{
-							FString Label;
-							FString Detail;
-							float Percent;
-							if (BuildActivity->Transfer.GetCurrentProgress(Label, Detail, Percent))
-							{
-								return FText::FromString(Detail.IsEmpty() ? Label : Detail);
-							}
-							return FText::FromString(FString());
-						})
-						.TextStyle(&FAppStyle::GetWidgetStyle<FTextBlockStyle>(TEXT("SmallText")))
-						.Justification(ETextJustify::Type::Left)
-						.WrapTextAt(0.f)
-						.ToolTipText_Lambda([this]()
-						{
-							return FText::FromString(BuildActivity->Transfer.GetRecentOutput());
-						})
-						.Visibility_Lambda([this]()
+		return SNew(SHorizontalBox)
+				+ SHorizontalBox::Slot()
+				.Padding(1.f,1.f)
+				.VAlign(VAlign_Center)
+				[
+					SNew(STextBlock)
+					.Text_Lambda([this]()
+					{
+						FString Label;
+						FString Detail;
+						float Percent;
+						if (BuildActivity->Transfer.GetCurrentProgress(Label, Detail, Percent))
 						{
-							return BuildActivity->Transfer.GetStatus() == FBuildServiceInstance::EBuildTransferStatus::Active ? EVisibility::Visible : EVisibility::Collapsed;
-						})
-					];
+							return FText::FromString(Detail.IsEmpty() ? Label : Detail);
+						}
+						return FText::FromString(FString());
+					})
+					.TextStyle(&FAppStyle::GetWidgetStyle<FTextBlockStyle>(TEXT("SmallText")))
+					.Justification(ETextJustify::Type::Left)
+					.WrapTextAt(0.f)
+					.ToolTipText_Lambda([this]()
+					{
+						return FText::FromString(BuildActivity->Transfer.GetRecentOutput());
+					})
+					.Visibility_Lambda([this]()
+					{
+						return BuildActivity->Transfer.GetStatus() == FBuildServiceInstance::EBuildTransferStatus::Active ? EVisibility::Visible : EVisibility::Collapsed;
+					})
+				];
 	}
 	else if (ColumnName == FBuildActivityIds::ColActions)
 	{
diff --git a/Engine/Source/Editor/StorageServerWidgets/Private/SBuildArtifactContentsViewer.cpp b/Engine/Source/Editor/StorageServerWidgets/Private/SBuildArtifactContentsViewer.cpp
--- a/Engine/Source/Editor/StorageServerWidgets/Private/SBuildArtifactContentsViewer.cpp
+++ b/Engine/Source/Editor/StorageServerWidgets/Private/SBuildArtifactContentsViewer.cpp
@@ -949,23 +949,33 @@ EActiveTimerReturnType SBuildArtifactContentsViewer::PollPendingFileOpens(double
 				for (const TSharedPtr<FContentsEntry>& Entry : Pending.Entries)
 				{
 					FString FullPath = FPaths::Combine(Pending.Transfer.GetDestination(), Entry->Name);
-					FPlatformProcess::LaunchFileInDefaultExternalApplication(*FullPath);
+					FPlatformProcess::LaunchFileInDefaultExternalApplication(*FullPath, nullptr, ELaunchVerb::Edit);
 					Entry->bIsDownloading = false;
 				}
 				break;
 
 			case EFileOpenAction::ExploreFolder:
 				{
-					TSet<FString> FoldersToExplore;
+					FString CommonPath;
 					for (const TSharedPtr<FContentsEntry>& Entry : Pending.Entries)
 					{
-						FString FullPath = FPaths::Combine(Pending.Transfer.GetDestination(), Entry->Name);
-						FoldersToExplore.Add(FPaths::GetPath(FullPath));
+						FString Folder = FPaths::GetPath(FPaths::Combine(Pending.Transfer.GetDestination(), Entry->Name));
+						if (CommonPath.IsEmpty())
+						{
+							CommonPath = Folder;
+						}
+						else if (!FPaths::IsSamePath(CommonPath, Folder))
+						{
+							while (!FPaths::IsUnderDirectory(Folder, CommonPath))
+							{
+								CommonPath = FPaths::GetPath(CommonPath);
+							}
+						}
 						Entry->bIsDownloading = false;
 					}
-					for (const FString& Folder : FoldersToExplore)
+					if (!CommonPath.IsEmpty())
 					{
-						FPlatformProcess::ExploreFolder(*Folder);
+						FPlatformProcess::ExploreFolder(*CommonPath);
 					}
 				}
 				break;
diff --git a/Engine/Source/Editor/StorageServerWidgets/Private/SBuildArtifactSelection.cpp b/Engine/Source/Editor/StorageServerWidgets/Private/SBuildArtifactSelection.cpp
--- a/Engine/Source/Editor/StorageServerWidgets/Private/SBuildArtifactSelection.cpp
+++ b/Engine/Source/Editor/StorageServerWidgets/Private/SBuildArtifactSelection.cpp
@@ -29,7 +29,7 @@ namespace
 		static FString TempRoot = FPaths::Combine(
 			FPlatformProcess::UserTempDir(),
 			TEXT("BuildStorageTool"),
-			TEXT("Open"));
+			TEXT("Scratch"));
 		static bool bCleanupRegistered = false;
 		if (!bCleanupRegistered)
 		{
diff --git a/Engine/Source/Programs/BuildStorageTool/Private/BuildStorageTool.cpp b/Engine/Source/Programs/BuildStorageTool/Private/BuildStorageTool.cpp
--- a/Engine/Source/Programs/BuildStorageTool/Private/BuildStorageTool.cpp
+++ b/Engine/Source/Programs/BuildStorageTool/Private/BuildStorageTool.cpp
@@ -623,11 +623,16 @@ class FBuildStorageToolApp
 			UE::Zen::Build::FBuildReference BuildReference;
 			if (TryParseBuildReferenceFromUrl(UrlMessage.Url, BuildReference))
 			{
-				BuildSelection->ActOnBuildArtifact(
-					UrlMessage.Action == EUrlAction::HighlightAndDownload ?
-						SBuildSelection::EBuildArtifactAction::HighlightAndDownload :
-						SBuildSelection::EBuildArtifactAction::Highlight,
-					BuildReference);
+				SBuildSelection::EBuildArtifactAction Action = UrlMessage.Action == EUrlAction::HighlightAndDownload
+					? SBuildSelection::EBuildArtifactAction::HighlightAndDownload
+					: SBuildSelection::EBuildArtifactAction::Highlight;
+
+				AsyncTask(ENamedThreads::GameThread, [this, BuildReference = MoveTemp(BuildReference), Action]()
+				{
+					Window->HACK_ForceToFront();
+					Window->DrawAttention(FWindowDrawAttentionParameters());
+					BuildSelection->ActOnBuildArtifact(Action, BuildReference);
+				});
 			}
 		}
 	}
@@ -668,7 +673,7 @@ public:
 			.ActivationPolicy(EWindowActivationPolicy::Always)
 			.SizingRule(ESizingRule::UserSized)
 			.IsTopmostWindow(false)
-			.FocusWhenFirstShown(false)
+			.FocusWhenFirstShown(true)
 			.SupportsMaximize(true)
 			.SupportsMinimize(true)
 			.HasCloseButton(true)
