BuildStorageTool: Refactor BuildArtifactContentsViewer to tabbed per-part layout

#Synced-CL 52440464

diff --git a/Engine/Source/Editor/StorageServerWidgets/Internal/SBuildArtifactSelection.h b/Engine/Source/Editor/StorageServerWidgets/Internal/SBuildArtifactSelection.h
--- a/Engine/Source/Editor/StorageServerWidgets/Internal/SBuildArtifactSelection.h
+++ b/Engine/Source/Editor/StorageServerWidgets/Internal/SBuildArtifactSelection.h
@@ -18,7 +18,7 @@
 
 #define UE_API STORAGESERVERWIDGETS_API
 
-DECLARE_DELEGATE_ThreeParams(FOnDownloadWithSpec, const FString& /*ArtifactName*/, const UE::BuildSelection::Internal::FArtifact& /*Artifact*/, const FString& /*DownloadSpecJSON*/);
+DECLARE_DELEGATE_FourParams(FOnDownloadWithSpec, const FString& /*ArtifactName*/, const UE::BuildSelection::Internal::FArtifact& /*Artifact*/, const FString& /*DownloadSpecJSON*/, const TArray<FString>& /*FullyMarkedPartNames*/);
 
 class SBuildArtifactSelection : public SCompoundWidget
 {
diff --git a/Engine/Source/Editor/StorageServerWidgets/Internal/SBuildSelection.h b/Engine/Source/Editor/StorageServerWidgets/Internal/SBuildSelection.h
--- a/Engine/Source/Editor/StorageServerWidgets/Internal/SBuildSelection.h
+++ b/Engine/Source/Editor/StorageServerWidgets/Internal/SBuildSelection.h
@@ -144,7 +144,7 @@ private:
 
 	FReply ExploreDestination_OnClicked();
 	void DownloadSelectedArtifactsInSelectedBuildGroup();
-	bool DownloadItem(const FBuildGroup& Group, const FString& ArtifactName, const UE::BuildSelection::Internal::FArtifact& Artifact, FString&& DownloadSpecJSONContents = FString());
+	bool DownloadItem(const FBuildGroup& Group, const FString& ArtifactName, const UE::BuildSelection::Internal::FArtifact& Artifact, FString&& DownloadSpecJSONContents = FString(), TArray<FString>&& PartNames = TArray<FString>());
 
 	TMap<FString, FString> EngineInstallations;
 
diff --git a/Engine/Source/Editor/StorageServerWidgets/Private/SBuildArtifactContentsViewer.cpp b/Engine/Source/Editor/StorageServerWidgets/Private/SBuildArtifactContentsViewer.cpp
--- a/Engine/Source/Editor/StorageServerWidgets/Private/SBuildArtifactContentsViewer.cpp
+++ b/Engine/Source/Editor/StorageServerWidgets/Private/SBuildArtifactContentsViewer.cpp
@@ -13,6 +13,8 @@
 #include "Styling/StyleColors.h"
 #include "Styling/AppStyle.h"
 #include "Widgets/Images/SThrobber.h"
+#include "Styling/CoreStyle.h"
+#include "Widgets/Layout/SBox.h"
 
 #define LOCTEXT_NAMESPACE "BuildArtifactContents"
 
@@ -23,6 +25,28 @@ namespace UE::BuildArtifactContents::Internal
 	const FName ColSize = TEXT("Size");
 }
 
+// ---------------------------------------------------------------------------
+// FPartTab
+// ---------------------------------------------------------------------------
+
+ECheckBoxState FPartTab::GetCheckState() const
+{
+	if (MarkedFileCount == 0)
+	{
+		return ECheckBoxState::Unchecked;
+	}
+	if (MarkedFileCount == TotalFileCount)
+	{
+		return ECheckBoxState::Checked;
+	}
+	return ECheckBoxState::Undetermined;
+}
+
+bool FPartTab::IsFullyMarked() const
+{
+	return TotalFileCount > 0 && MarkedFileCount == TotalFileCount;
+}
+
 // ---------------------------------------------------------------------------
 // SContentsEntryRow
 // ---------------------------------------------------------------------------
@@ -38,19 +62,6 @@ TSharedRef<SWidget> SContentsEntryRow::GenerateWidgetForColumn(const FName& Colu
 {
 	using namespace UE::BuildArtifactContents::Internal;
 
-	if (Entry->bIsPartHeader)
-	{
-		if (ColumnName == ColName)
-		{
-			return SNew(STextBlock)
-				.Text(FText::FromString(Entry->Name))
-				.Font(FCoreStyle::GetDefaultFontStyle("Bold", 9))
-				.ColorAndOpacity(FStyleColors::AccentGreen)
-				.Margin(FMargin(4.0f, 4.0f));
-		}
-		return SNullWidget::NullWidget;
-	}
-
 	if (ColumnName == ColMark)
 	{
 		return SNew(SCheckBox)
@@ -119,6 +130,31 @@ void SBuildArtifactContentsViewer::Construct(const FArguments& InArgs)
 
 	OnStartFileDownload = InArgs._OnStartFileDownload;
 
+	// Initialize tab button styles from the docking tab style
+	{
+		const FDockTabStyle& DockTabStyle = FCoreStyle::Get().GetWidgetStyle<FDockTabStyle>("Docking.Tab");
+
+		InactiveTabButtonStyle = FButtonStyle()
+			.SetNormal(DockTabStyle.NormalBrush)
+			.SetHovered(DockTabStyle.HoveredBrush)
+			.SetPressed(DockTabStyle.ForegroundBrush)
+			.SetNormalForeground(DockTabStyle.NormalForegroundColor)
+			.SetHoveredForeground(DockTabStyle.HoveredForegroundColor)
+			.SetPressedForeground(DockTabStyle.ActiveForegroundColor)
+			.SetNormalPadding(DockTabStyle.TabPadding)
+			.SetPressedPadding(DockTabStyle.TabPadding);
+
+		ActiveTabButtonStyle = FButtonStyle()
+			.SetNormal(DockTabStyle.ForegroundBrush)
+			.SetHovered(DockTabStyle.ForegroundBrush)
+			.SetPressed(DockTabStyle.ForegroundBrush)
+			.SetNormalForeground(DockTabStyle.ActiveForegroundColor)
+			.SetHoveredForeground(DockTabStyle.ActiveForegroundColor)
+			.SetPressedForeground(DockTabStyle.ActiveForegroundColor)
+			.SetNormalPadding(DockTabStyle.TabPadding)
+			.SetPressedPadding(DockTabStyle.TabPadding);
+	}
+
 	ChildSlot
 	[
 		SNew(SVerticalBox)
@@ -188,7 +224,16 @@ void SBuildArtifactContentsViewer::Construct(const FArguments& InArgs)
 			]
 		]
 
-		// Contents list
+		// Tab bar
+		+ SVerticalBox::Slot()
+		.AutoHeight()
+		.Padding(4.0f, 0.0f)
+		[
+			SAssignNew(TabBarScrollBox, SScrollBox)
+			.Orientation(EOrientation::Orient_Horizontal)
+		]
+
+		// Contents area with widget switcher
 		+ SVerticalBox::Slot()
 		.FillHeight(1.0f)
 		.Padding(4.0f)
@@ -197,33 +242,7 @@ void SBuildArtifactContentsViewer::Construct(const FArguments& InArgs)
 			.BorderImage(FAppStyle::GetBrush("ToolPanel.GroupBorder"))
 			.Padding(0.0f)
 			[
-				SAssignNew(ListView, SListView<TSharedPtr<FContentsEntry>>)
-				.ListItemsSource(&FilteredEntries)
-				.OnGenerateRow(this, &SBuildArtifactContentsViewer::OnGenerateRow)
-				.OnContextMenuOpening(this, &SBuildArtifactContentsViewer::OnContextMenuOpening)
-				.SelectionMode(ESelectionMode::Multi)
-				.HeaderRow
-				(
-					SNew(SHeaderRow)
-					+ SHeaderRow::Column(ColMark)
-					.DefaultLabel(FText::GetEmpty())
-					.FixedWidth(24.0f)
-					.HAlignHeader(HAlign_Center)
-					.HAlignCell(HAlign_Center)
-
-					+ SHeaderRow::Column(ColName)
-					.DefaultLabel(LOCTEXT("NameColumn", "Name"))
-					.FillWidth(0.85f)
-					.SortMode(this, &SBuildArtifactContentsViewer::GetColumnSortMode, ColName)
-					.OnSort(this, &SBuildArtifactContentsViewer::OnColumnSortModeChanged)
-
-					+ SHeaderRow::Column(ColSize)
-					.DefaultLabel(LOCTEXT("SizeColumn", "Size"))
-					.FillWidth(0.15f)
-					.HAlignCell(HAlign_Right)
-					.SortMode(this, &SBuildArtifactContentsViewer::GetColumnSortMode, ColSize)
-					.OnSort(this, &SBuildArtifactContentsViewer::OnColumnSortModeChanged)
-				)
+				SAssignNew(ContentSwitcher, SWidgetSwitcher)
 			]
 		]
 
@@ -285,7 +304,9 @@ EActiveTimerReturnType SBuildArtifactContentsViewer::PollPendingContents(double
 
 void SBuildArtifactContentsViewer::SetContents(TMap<FString, UE::Zen::Build::FBuildServiceInstance::FBuildPart>& BuildContents)
 {
-	AllEntries.Empty();
+	using namespace UE::BuildArtifactContents::Internal;
+
+	PartTabs.Empty();
 	TotalFileCount = 0;
 	TotalFileSize = 0;
 
@@ -294,18 +315,16 @@ void SBuildArtifactContentsViewer::SetContents(TMap<FString, UE::Zen::Build::FBu
 	BuildContents.GetKeys(PartNames);
 	PartNames.Sort();
 
+	// Reserve so that PartTabs elements (and their PartName strings) stay at
+	// stable addresses while we populate FContentsEntry::PartName views below.
+	PartTabs.Reserve(PartNames.Num());
+
 	for (const FString& PartName : PartNames)
 	{
 		UE::Zen::Build::FBuildServiceInstance::FBuildPart& Part = BuildContents[PartName];
-		FName PartFName(*PartName);
-
-		// Add part header entry
-		TSharedPtr<FContentsEntry> HeaderEntry = MakeShared<FContentsEntry>();
-		HeaderEntry->bIsPartHeader = true;
-		HeaderEntry->Name = PartName;
-		HeaderEntry->PartName = PartFName;
-		HeaderEntry->PartId = Part.Id;
-		AllEntries.Add(HeaderEntry);
+		FPartTab& Tab = PartTabs.AddDefaulted_GetRef();
+		Tab.PartName = PartName;
+		Tab.PartId = Part.Id;
 
 		// Sort files by name
 		Part.Files.Sort([](const UE::Zen::Build::FBuildServiceInstance::FBuildFile& A, const UE::Zen::Build::FBuildServiceInstance::FBuildFile& B)
@@ -316,23 +335,185 @@ void SBuildArtifactContentsViewer::SetContents(TMap<FString, UE::Zen::Build::FBu
 		for (const UE::Zen::Build::FBuildServiceInstance::FBuildFile& File : Part.Files)
 		{
 			TSharedPtr<FContentsEntry> FileEntry = MakeShared<FContentsEntry>();
-			FileEntry->bIsPartHeader = false;
 			FileEntry->Name = File.Name;
-			FileEntry->PartName = PartFName;
+			FileEntry->PartName = Tab.PartName;
 			FileEntry->PartId = Part.Id;
 			FileEntry->RawSize = File.RawSize;
-			AllEntries.Add(FileEntry);
+			Tab.AllFiles.Add(FileEntry);
+
+			Tab.TotalFileSize += File.RawSize;
+			Tab.TotalFileCount++;
+		}
+
+		TotalFileSize += Tab.TotalFileSize;
+		TotalFileCount += Tab.TotalFileCount;
+	}
+
+	// Build list views and add to content switcher
+	if (ContentSwitcher.IsValid())
+	{
+		for (int32 TabIndex = 0; TabIndex < PartTabs.Num(); ++TabIndex)
+		{
+			FPartTab& Tab = PartTabs[TabIndex];
+			Tab.FilteredFiles = Tab.AllFiles;
+
+			ContentSwitcher->AddSlot()
+			[
+				SAssignNew(Tab.ListView, SListView<TSharedPtr<FContentsEntry>>)
+				.ListItemsSource(&Tab.FilteredFiles)
+				.OnGenerateRow(this, &SBuildArtifactContentsViewer::OnGenerateRow)
+				.OnContextMenuOpening(this, &SBuildArtifactContentsViewer::OnContextMenuOpening)
+				.SelectionMode(ESelectionMode::Multi)
+				.HeaderRow
+				(
+					SNew(SHeaderRow)
+					+ SHeaderRow::Column(ColMark)
+					.DefaultLabel(FText::GetEmpty())
+					.FixedWidth(24.0f)
+					.HAlignHeader(HAlign_Center)
+					.HAlignCell(HAlign_Center)
 
-			TotalFileSize += File.RawSize;
-			TotalFileCount++;
+					+ SHeaderRow::Column(ColName)
+					.DefaultLabel(LOCTEXT("NameColumn", "Name"))
+					.FillWidth(0.85f)
+					.SortMode(this, &SBuildArtifactContentsViewer::GetColumnSortMode, ColName)
+					.OnSort(this, &SBuildArtifactContentsViewer::OnColumnSortModeChanged)
+
+					+ SHeaderRow::Column(ColSize)
+					.DefaultLabel(LOCTEXT("SizeColumn", "Size"))
+					.FillWidth(0.15f)
+					.HAlignCell(HAlign_Right)
+					.SortMode(this, &SBuildArtifactContentsViewer::GetColumnSortMode, ColSize)
+					.OnSort(this, &SBuildArtifactContentsViewer::OnColumnSortModeChanged)
+				)
+			];
 		}
 	}
 
+	// Build tab bar
+	RebuildTabBar();
+
+	// Activate first tab
+	ActiveTabIndex = 0;
+	if (ContentSwitcher.IsValid() && PartTabs.Num() > 0)
+	{
+		ContentSwitcher->SetActiveWidgetIndex(0);
+	}
+
 	bIsLoading = false;
 
 	SortAndRefreshList();
 }
 
+void SBuildArtifactContentsViewer::RebuildTabBar()
+{
+	if (!TabBarScrollBox.IsValid())
+	{
+		return;
+	}
+
+	TabBarScrollBox->ClearChildren();
+
+	for (int32 TabIndex = 0; TabIndex < PartTabs.Num(); ++TabIndex)
+	{
+		FPartTab& Tab = PartTabs[TabIndex];
+		const bool bIsActive = (TabIndex == ActiveTabIndex);
+
+		TabBarScrollBox->AddSlot()
+		.Padding(0.0f, 0.0f, 2.0f, 0.0f)
+		[
+			SNew(SVerticalBox)
+			+ SVerticalBox::Slot()
+			.AutoHeight()
+			[
+				SNew(SButton)
+				.ButtonStyle(bIsActive ? &ActiveTabButtonStyle : &InactiveTabButtonStyle)
+				.OnClicked_Lambda([this, TabIndex]()
+				{
+					SetActiveTab(TabIndex);
+					return FReply::Handled();
+				})
+				[
+					SNew(SHorizontalBox)
+					+ SHorizontalBox::Slot()
+					.AutoWidth()
+					.VAlign(VAlign_Center)
+					.Padding(0.0f, 0.0f, 4.0f, 0.0f)
+					[
+						SNew(SCheckBox)
+						.IsChecked_Lambda([this, TabIndex]()
+						{
+							if (PartTabs.IsValidIndex(TabIndex))
+							{
+								return PartTabs[TabIndex].GetCheckState();
+							}
+							return ECheckBoxState::Unchecked;
+						})
+						.OnCheckStateChanged_Lambda([this, TabIndex](ECheckBoxState NewState)
+						{
+							if (!PartTabs.IsValidIndex(TabIndex))
+							{
+								return;
+							}
+							FPartTab& ClickedTab = PartTabs[TabIndex];
+							// Indeterminate or unchecked -> check all; checked -> uncheck all
+							const bool bMark = (NewState != ECheckBoxState::Unchecked) || (ClickedTab.GetCheckState() == ECheckBoxState::Undetermined);
+							for (const TSharedPtr<FContentsEntry>& Entry : ClickedTab.AllFiles)
+							{
+								Entry->bIsMarked = bMark;
+							}
+							RecalculateMarkedCounts();
+						})
+					]
+					+ SHorizontalBox::Slot()
+					.AutoWidth()
+					.VAlign(VAlign_Center)
+					[
+						SNew(STextBlock)
+						.Text(FText::FromString(Tab.PartName))
+						.Font(bIsActive
+							? FCoreStyle::GetDefaultFontStyle("Bold", 9)
+							: FCoreStyle::GetDefaultFontStyle("Regular", 9))
+					]
+				]
+			]
+			+ SVerticalBox::Slot()
+			.AutoHeight()
+			[
+				SNew(SBox)
+				.HeightOverride(2.f)
+				.Visibility(bIsActive ? EVisibility::SelfHitTestInvisible : EVisibility::Collapsed)
+				[
+					SNew(SBorder)
+					.BorderImage(FCoreStyle::Get().GetBrush("WhiteBrush"))
+					.BorderBackgroundColor(FStyleColors::Primary)
+					.Padding(0)
+				]
+			]
+		];
+	}
+}
+
+void SBuildArtifactContentsViewer::SetActiveTab(int32 Index)
+{
+	if (!PartTabs.IsValidIndex(Index))
+	{
+		return;
+	}
+
+	ActiveTabIndex = Index;
+	if (ContentSwitcher.IsValid())
+	{
+		ContentSwitcher->SetActiveWidgetIndex(Index);
+	}
+
+	// Rebuild tab bar to update active/inactive visual state
+	RebuildTabBar();
+
+	// Refresh status text for the newly active tab
+	RefreshFilteredList();
+}
+
 TSharedRef<ITableRow> SBuildArtifactContentsViewer::OnGenerateRow(TSharedPtr<FContentsEntry> Entry, const TSharedRef<STableViewBase>& OwnerTable)
 {
 	return SNew(SContentsEntryRow, OwnerTable, Entry)
@@ -385,32 +566,11 @@ void SBuildArtifactContentsViewer::SortAndRefreshList()
 
 	if (SortMode != EColumnSortMode::None && !SortByColumn.IsNone())
 	{
-		// Separate part headers from file entries, sort files, then rebuild with headers
-		struct FPartGroup
-		{
-			TSharedPtr<FContentsEntry> Header;
-			TArray<TSharedPtr<FContentsEntry>> Files;
-		};
-
-		TArray<FPartGroup> Groups;
-		for (const TSharedPtr<FContentsEntry>& Entry : AllEntries)
-		{
-			if (Entry->bIsPartHeader)
-			{
-				Groups.AddDefaulted();
-				Groups.Last().Header = Entry;
-			}
-			else if (Groups.Num() > 0)
-			{
-				Groups.Last().Files.Add(Entry);
-			}
-		}
-
 		const bool bAscending = (SortMode == EColumnSortMode::Ascending);
 
-		for (FPartGroup& Group : Groups)
+		for (FPartTab& Tab : PartTabs)
 		{
-			Group.Files.Sort([&](const TSharedPtr<FContentsEntry>& A, const TSharedPtr<FContentsEntry>& B)
+			Tab.AllFiles.Sort([&](const TSharedPtr<FContentsEntry>& A, const TSharedPtr<FContentsEntry>& B)
 			{
 				if (SortByColumn == ColName)
 				{
@@ -423,13 +583,6 @@ void SBuildArtifactContentsViewer::SortAndRefreshList()
 				return false;
 			});
 		}
-
-		AllEntries.Empty();
-		for (const FPartGroup& Group : Groups)
-		{
-			AllEntries.Add(Group.Header);
-			AllEntries.Append(Group.Files);
-		}
 	}
 
 	RefreshFilteredList();
@@ -437,52 +590,42 @@ void SBuildArtifactContentsViewer::SortAndRefreshList()
 
 void SBuildArtifactContentsViewer::RefreshFilteredList()
 {
-	FilteredEntries.Empty();
+	TArray<FString> IncludeWildcards;
+	if (!SearchText.IsEmpty())
+	{
+		ParseIncludeWildcards(SearchText.ToString(), IncludeWildcards);
+	}
 
 	int32 FilteredFileCount = 0;
 	uint64 FilteredFileSize = 0;
 
-	if (SearchText.IsEmpty())
+	for (FPartTab& Tab : PartTabs)
 	{
-		FilteredEntries = AllEntries;
-		FilteredFileCount = TotalFileCount;
-		FilteredFileSize = TotalFileSize;
-	}
-	else
-	{
-		TArray<FString> IncludeWildcards;
-		ParseIncludeWildcards(SearchText.ToString(), IncludeWildcards);
+		Tab.FilteredFiles.Empty();
 
-		// Track which part headers should be included (only if they have matching files)
-		int32 PendingHeaderIndex = INDEX_NONE;
-		for (int32 i = 0; i < AllEntries.Num(); ++i)
+		if (SearchText.IsEmpty())
 		{
-			const TSharedPtr<FContentsEntry>& Entry = AllEntries[i];
-			if (Entry->bIsPartHeader)
-			{
-				PendingHeaderIndex = i;
-			}
-			else
+			Tab.FilteredFiles = Tab.AllFiles;
+			FilteredFileCount += Tab.TotalFileCount;
+			FilteredFileSize += Tab.TotalFileSize;
+		}
+		else
+		{
+			for (const TSharedPtr<FContentsEntry>& Entry : Tab.AllFiles)
 			{
 				if (MatchesIncludeWildcards(IncludeWildcards, Entry->Name))
 				{
-					// Include the part header if we haven't already
-					if (PendingHeaderIndex != INDEX_NONE)
-					{
-						FilteredEntries.Add(AllEntries[PendingHeaderIndex]);
-						PendingHeaderIndex = INDEX_NONE;
-					}
-					FilteredEntries.Add(Entry);
+					Tab.FilteredFiles.Add(Entry);
 					FilteredFileCount++;
 					FilteredFileSize += Entry->RawSize;
 				}
 			}
 		}
-	}
 
-	if (ListView.IsValid())
-	{
-		ListView->RequestListRefresh();
+		if (Tab.ListView.IsValid())
+		{
+			Tab.ListView->RequestListRefresh();
+		}
 	}
 
 	if (StatusText.IsValid() && !bIsLoading)
@@ -638,34 +781,49 @@ bool SBuildArtifactContentsViewer::MatchesIncludeWildcards(const TArray<FString>
 
 void SBuildArtifactContentsViewer::MarkSelectedFiles()
 {
-	if (!ListView.IsValid())
+	if (!PartTabs.IsValidIndex(ActiveTabIndex))
+	{
+		return;
+	}
+
+	FPartTab& ActiveTab = PartTabs[ActiveTabIndex];
+	if (!ActiveTab.ListView.IsValid())
 	{
 		return;
 	}
 
-	TArray<TSharedPtr<FContentsEntry>> SelectedItems = ListView->GetSelectedItems();
+	TArray<TSharedPtr<FContentsEntry>> SelectedItems = ActiveTab.ListView->GetSelectedItems();
 	for (const TSharedPtr<FContentsEntry>& Item : SelectedItems)
 	{
-		if (!Item->bIsPartHeader)
-		{
-			Item->bIsMarked = !Item->bIsMarked;
-		}
+		Item->bIsMarked = !Item->bIsMarked;
 	}
 
 	RecalculateMarkedCounts();
 }
 
+void SBuildArtifactContentsViewer::RecalculatePartMarkedCounts(FPartTab& Tab)
+{
+	Tab.MarkedFileCount = 0;
+	Tab.MarkedFileSize = 0;
+	for (const TSharedPtr<FContentsEntry>& Entry : Tab.AllFiles)
+	{
+		if (Entry->bIsMarked)
+		{
+			Tab.MarkedFileCount++;
+			Tab.MarkedFileSize += Entry->RawSize;
+		}
+	}
+}
+
 void SBuildArtifactContentsViewer::RecalculateMarkedCounts()
 {
 	MarkedFileCount = 0;
 	MarkedFileSize = 0;
-	for (const TSharedPtr<FContentsEntry>& Entry : AllEntries)
+	for (FPartTab& Tab : PartTabs)
 	{
-		if (!Entry->bIsPartHeader && Entry->bIsMarked)
-		{
-			MarkedFileCount++;
-			MarkedFileSize += Entry->RawSize;
-		}
+		RecalculatePartMarkedCounts(Tab);
+		MarkedFileCount += Tab.MarkedFileCount;
+		MarkedFileSize += Tab.MarkedFileSize;
 	}
 
 	if (MarkedStatusText.IsValid())
@@ -680,17 +838,20 @@ void SBuildArtifactContentsViewer::RecalculateMarkedCounts()
 		}
 	}
 
-	if (ListView.IsValid())
+	for (FPartTab& Tab : PartTabs)
 	{
-		ListView->RequestListRefresh();
+		if (Tab.ListView.IsValid())
+		{
+			Tab.ListView->RequestListRefresh();
+		}
 	}
 }
 
 void SBuildArtifactContentsViewer::MarkAllFiles()
 {
-	for (const TSharedPtr<FContentsEntry>& Entry : AllEntries)
+	for (FPartTab& Tab : PartTabs)
 	{
-		if (!Entry->bIsPartHeader)
+		for (const TSharedPtr<FContentsEntry>& Entry : Tab.AllFiles)
 		{
 			Entry->bIsMarked = true;
 		}
@@ -701,17 +862,25 @@ void SBuildArtifactContentsViewer::MarkAllFiles()
 
 void SBuildArtifactContentsViewer::ClearAllMarkedFiles()
 {
-	for (const TSharedPtr<FContentsEntry>& Entry : AllEntries)
+	for (FPartTab& Tab : PartTabs)
 	{
-		Entry->bIsMarked = false;
+		for (const TSharedPtr<FContentsEntry>& Entry : Tab.AllFiles)
+		{
+			Entry->bIsMarked = false;
+		}
 	}
 
 	MarkedFileCount = 0;
 	MarkedFileSize = 0;
 
-	if (ListView.IsValid())
+	for (FPartTab& Tab : PartTabs)
 	{
-		ListView->RequestListRefresh();
+		Tab.MarkedFileCount = 0;
+		Tab.MarkedFileSize = 0;
+		if (Tab.ListView.IsValid())
+		{
+			Tab.ListView->RequestListRefresh();
+		}
 	}
 }
 
@@ -725,84 +894,100 @@ FReply SBuildArtifactContentsViewer::OnKeyDown(const FGeometry& MyGeometry, cons
 	return SCompoundWidget::OnKeyDown(MyGeometry, InKeyEvent);
 }
 
-FString SBuildArtifactContentsViewer::ComposeDownloadSpecJSON() const
+FDownloadSpec SBuildArtifactContentsViewer::ComposeDownloadSpec() const
 {
-	// Group marked files by their parent part
+	FDownloadSpec Result;
+
 	struct FPartInfo
 	{
 		FCbObjectId PartId;
 		TArray<FString> Files;
 	};
 
-	TArray<TPair<FString, FPartInfo>> Parts;
+	TArray<TPair<FString, FPartInfo>> PartialParts;
 
-	for (const TSharedPtr<FContentsEntry>& Entry : AllEntries)
+	for (const FPartTab& Tab : PartTabs)
 	{
-		if (Entry->bIsPartHeader || !Entry->bIsMarked)
+		if (Tab.MarkedFileCount == 0)
 		{
 			continue;
 		}
 
-		FString PartNameStr = Entry->PartName.ToString();
-
-		// Find or create the part group
-		FPartInfo* Found = nullptr;
-		for (TPair<FString, FPartInfo>& Pair : Parts)
+		if (Tab.IsFullyMarked())
 		{
-			if (Pair.Key == PartNameStr)
-			{
-				Found = &Pair.Value;
-				break;
-			}
+			Result.FullyMarkedPartNames.Add(Tab.PartName);
 		}
-		if (!Found)
+		else
 		{
-			Parts.Emplace(PartNameStr, FPartInfo{ Entry->PartId, {} });
-			Found = &Parts.Last().Value;
+			FPartInfo Info;
+			Info.PartId = Tab.PartId;
+			for (const TSharedPtr<FContentsEntry>& Entry : Tab.AllFiles)
+			{
+				if (Entry->bIsMarked)
+				{
+					Info.Files.Add(Entry->Name);
+				}
+			}
+			PartialParts.Emplace(Tab.PartName, MoveTemp(Info));
 		}
-		Found->Files.Add(Entry->Name);
 	}
 
-	FString OutputString;
-	TSharedRef<TJsonWriter<>> Writer = TJsonWriterFactory<>::Create(&OutputString);
-
-	Writer->WriteObjectStart();
-	Writer->WriteObjectStart(TEXT("parts"));
-	for (const TPair<FString, FPartInfo>& Pair : Parts)
+	// Build JSON only for partially-marked parts; skip entirely when all
+	// marked parts are fully marked since PartNames alone is sufficient.
+	if (!PartialParts.IsEmpty())
 	{
-		Writer->WriteObjectStart(Pair.Key);
-		if (!(Pair.Value.PartId == FCbObjectId()))
-		{
-			Writer->WriteValue(TEXT("partId"), FString(WriteToString<64>(Pair.Value.PartId)));
-		}
-		Writer->WriteArrayStart(TEXT("files"));
-		for (const FString& FileName : Pair.Value.Files)
+		FString OutputString;
+		TSharedRef<TJsonWriter<>> Writer = TJsonWriterFactory<>::Create(&OutputString);
+
+		Writer->WriteObjectStart();
+		Writer->WriteObjectStart(TEXT("parts"));
+		for (const TPair<FString, FPartInfo>& Pair : PartialParts)
 		{
-			Writer->WriteValue(FileName);
+			Writer->WriteObjectStart(Pair.Key);
+			if (!(Pair.Value.PartId == FCbObjectId()))
+			{
+				Writer->WriteValue(TEXT("partId"), FString(WriteToString<64>(Pair.Value.PartId)));
+			}
+			Writer->WriteArrayStart(TEXT("files"));
+			for (const FString& FileName : Pair.Value.Files)
+			{
+				Writer->WriteValue(FileName);
+			}
+			Writer->WriteArrayEnd();
+			Writer->WriteObjectEnd();
 		}
-		Writer->WriteArrayEnd();
 		Writer->WriteObjectEnd();
+		Writer->WriteObjectEnd();
+
+		Writer->Close();
+		Result.SpecJSON = MoveTemp(OutputString);
 	}
-	Writer->WriteObjectEnd();
-	Writer->WriteObjectEnd();
 
-	Writer->Close();
-	return OutputString;
+	return Result;
+}
+
+FString SBuildArtifactContentsViewer::ComposeDownloadSpecJSON() const
+{
+	return ComposeDownloadSpec().SpecJSON;
 }
 
 TSharedPtr<SWidget> SBuildArtifactContentsViewer::OnContextMenuOpening()
 {
-	TArray<TSharedPtr<FContentsEntry>> SelectedItems = ListView->GetSelectedItems();
+	if (!PartTabs.IsValidIndex(ActiveTabIndex))
+	{
+		return nullptr;
+	}
 
-	// Filter to non-header file entries only
-	TArray<TSharedPtr<FContentsEntry>> FileEntries;
-	for (const TSharedPtr<FContentsEntry>& Item : SelectedItems)
+	FPartTab& ActiveTab = PartTabs[ActiveTabIndex];
+	if (!ActiveTab.ListView.IsValid())
 	{
-		if (!Item->bIsPartHeader)
-		{
-			FileEntries.Add(Item);
-		}
+		return nullptr;
 	}
+
+	TArray<TSharedPtr<FContentsEntry>> SelectedItems = ActiveTab.ListView->GetSelectedItems();
+
+	// All entries are file entries now
+	TArray<TSharedPtr<FContentsEntry>> FileEntries = MoveTemp(SelectedItems);
 	if (FileEntries.IsEmpty() || !OnStartFileDownload.IsBound())
 	{
 		return nullptr;
@@ -885,7 +1070,7 @@ FString SBuildArtifactContentsViewer::ComposeFileDownloadSpec(const TArray<TShar
 
 	for (const TSharedPtr<FContentsEntry>& FileEntry : FileEntries)
 	{
-		FString PartNameStr = FileEntry->PartName.ToString();
+		FString PartNameStr(FileEntry->PartName);
 
 		// Find or create the part group
 		FPartInfo* Found = nullptr;
diff --git a/Engine/Source/Editor/StorageServerWidgets/Private/SBuildArtifactContentsViewer.h b/Engine/Source/Editor/StorageServerWidgets/Private/SBuildArtifactContentsViewer.h
--- a/Engine/Source/Editor/StorageServerWidgets/Private/SBuildArtifactContentsViewer.h
+++ b/Engine/Source/Editor/StorageServerWidgets/Private/SBuildArtifactContentsViewer.h
@@ -8,17 +8,40 @@
 #include "Widgets/SCompoundWidget.h"
 #include "Widgets/Views/SListView.h"
 #include "Widgets/Views/STableRow.h"
+#include "Widgets/Layout/SWidgetSwitcher.h"
+#include "Widgets/Layout/SScrollBox.h"
 #include "Experimental/BuildServerInterface.h"
 
 struct FContentsEntry
 {
-	bool bIsPartHeader = false;
-	bool bIsMarked = false;     // Persistent download mark (file entries only)
+	bool bIsMarked = false;     // Persistent download mark
 	bool bIsDownloading = false; // true while an "Open" download is in progress
-	FString Name;               // Part name for headers, file name for files
-	FName PartName;             // Parent part name (set on all entries)
-	FCbObjectId PartId;         // Parent part ID (set on all entries)
-	uint64 RawSize = 0;         // File entries only
+	FString Name;               // File name
+	FStringView PartName;       // Non-owning view into FPartTab::PartName (stable after SetContents)
+	FCbObjectId PartId;         // Parent part ID
+	uint64 RawSize = 0;
+};
+
+struct FPartTab
+{
+	FString PartName;
+	FCbObjectId PartId;
+	TArray<TSharedPtr<FContentsEntry>> AllFiles;
+	TArray<TSharedPtr<FContentsEntry>> FilteredFiles;
+	TSharedPtr<SListView<TSharedPtr<FContentsEntry>>> ListView;
+	int32 TotalFileCount = 0;
+	uint64 TotalFileSize = 0;
+	int32 MarkedFileCount = 0;
+	uint64 MarkedFileSize = 0;
+
+	ECheckBoxState GetCheckState() const;
+	bool IsFullyMarked() const;
+};
+
+struct FDownloadSpec
+{
+	FString SpecJSON;                        // JSON for partially-marked parts (individual files)
+	TArray<FString> FullyMarkedPartNames;    // Parts where every file is marked
 };
 
 /**
@@ -47,7 +70,7 @@ DECLARE_DELEGATE_RetVal_OneParam(
 );
 
 /**
- * Widget for viewing build artifact file contents in a searchable list grouped by build part
+ * Widget for viewing build artifact file contents in a tabbed view grouped by build part
  */
 class SBuildArtifactContentsViewer : public SCompoundWidget
 {
@@ -64,23 +87,26 @@ public:
 	/** Thread-safe: queues contents from a background thread for consumption on the Slate thread. */
 	void QueueContents(TMap<FString, UE::Zen::Build::FBuildServiceInstance::FBuildPart>&& BuildContents);
 
-	/** Toggle bIsMarked on selected non-header items and update counts. */
+	/** Toggle bIsMarked on selected items in the active tab and update counts. */
 	void MarkSelectedFiles();
 
-	/** Mark all file entries for download. */
+	/** Mark all file entries for download across all tabs. */
 	void MarkAllFiles();
 
-	/** Unmark all marked files. */
+	/** Unmark all marked files across all tabs. */
 	void ClearAllMarkedFiles();
 
 	bool HasMarkedFiles() const { return MarkedFileCount > 0; }
 	int32 GetMarkedFileCount() const { return MarkedFileCount; }
 	uint64 GetMarkedFileSize() const { return MarkedFileSize; }
 
-	/** Recompute MarkedFileCount/MarkedFileSize from AllEntries and update status text. */
+	/** Recompute MarkedFileCount/MarkedFileSize from all PartTabs and update status text. */
 	void RecalculateMarkedCounts();
 
-	/** Build download spec JSON from marked entries. */
+	/** Build download spec from marked entries, separating fully-marked parts. */
+	FDownloadSpec ComposeDownloadSpec() const;
+
+	/** Build download spec JSON from marked entries (thin wrapper). */
 	FString ComposeDownloadSpecJSON() const;
 
 	virtual FReply OnKeyDown(const FGeometry& MyGeometry, const FKeyEvent& InKeyEvent) override;
@@ -105,8 +131,10 @@ private:
 
 	TSharedRef<ITableRow> OnGenerateRow(TSharedPtr<FContentsEntry> Entry, const TSharedRef<STableViewBase>& OwnerTable);
 
+	void SetActiveTab(int32 Index);
+	void RecalculatePartMarkedCounts(FPartTab& Tab);
+	void RebuildTabBar();
 	void RefreshFilteredList();
-
 	void SortAndRefreshList();
 
 	void OnSearchTextChanged(const FText& InText);
@@ -121,11 +149,11 @@ private:
 	FCriticalSection PendingContentsMutex;
 	TOptional<TMap<FString, UE::Zen::Build::FBuildServiceInstance::FBuildPart>> PendingContents;
 
-	TArray<TSharedPtr<FContentsEntry>> AllEntries;
-
-	TArray<TSharedPtr<FContentsEntry>> FilteredEntries;
+	TArray<FPartTab> PartTabs;
 
-	TSharedPtr<SListView<TSharedPtr<FContentsEntry>>> ListView;
+	TSharedPtr<SWidgetSwitcher> ContentSwitcher;
+	TSharedPtr<SScrollBox> TabBarScrollBox;
+	int32 ActiveTabIndex = 0;
 
 	TSharedPtr<STextBlock> StatusText;
 	TSharedPtr<STextBlock> MarkedStatusText;
@@ -152,4 +180,7 @@ private:
 	FOnStartFileDownload OnStartFileDownload;
 
 	bool bIsLoading = true;
+
+	FButtonStyle InactiveTabButtonStyle;
+	FButtonStyle ActiveTabButtonStyle;
 };
diff --git a/Engine/Source/Editor/StorageServerWidgets/Private/SBuildArtifactSelection.cpp b/Engine/Source/Editor/StorageServerWidgets/Private/SBuildArtifactSelection.cpp
--- a/Engine/Source/Editor/StorageServerWidgets/Private/SBuildArtifactSelection.cpp
+++ b/Engine/Source/Editor/StorageServerWidgets/Private/SBuildArtifactSelection.cpp
@@ -703,8 +703,8 @@ TSharedRef<SWidget> SBuildArtifactSelection::CreateArtifactMenuWidget(UE::BuildS
 					int32 ButtonIndex = Dialog->ShowModal();
 					if (ButtonIndex == 0 && ViewerPtr->HasMarkedFiles())
 					{
-						FString DownloadSpec = ViewerPtr->ComposeDownloadSpecJSON();
-						OnDownloadWithSpec.ExecuteIfBound(ArtifactName, *Artifact, DownloadSpec);
+						FDownloadSpec Spec = ViewerPtr->ComposeDownloadSpec();
+						OnDownloadWithSpec.ExecuteIfBound(ArtifactName, *Artifact, Spec.SpecJSON, Spec.FullyMarkedPartNames);
 					}
 				})
 		),
diff --git a/Engine/Source/Editor/StorageServerWidgets/Private/SBuildSelection.cpp b/Engine/Source/Editor/StorageServerWidgets/Private/SBuildSelection.cpp
--- a/Engine/Source/Editor/StorageServerWidgets/Private/SBuildSelection.cpp
+++ b/Engine/Source/Editor/StorageServerWidgets/Private/SBuildSelection.cpp
@@ -1529,7 +1529,7 @@ void SBuildSelection::DownloadSelectedArtifactsInSelectedBuildGroup()
 	}
 }
 
-bool SBuildSelection::DownloadItem(const FBuildGroup& Group, const FString& ArtifactName, const UE::BuildSelection::Internal::FArtifact& Artifact, FString&& DownloadSpecJSONContents)
+bool SBuildSelection::DownloadItem(const FBuildGroup& Group, const FString& ArtifactName, const UE::BuildSelection::Internal::FArtifact& Artifact, FString&& DownloadSpecJSONContents, TArray<FString>&& PartNames)
 {
 	using namespace UE::Zen::Build;
 
@@ -1630,7 +1630,7 @@ bool SBuildSelection::DownloadItem(const FBuildGroup& Group, const FString& Arti
 				DownloadSpecJSONContents.IsEmpty() ? UserSelectedTransferFilters[BuildTypeIndex].ExcludeFilter : FString(),
 				MoveTemp(DownloadSpecJSONContents),
 				TArray<FString>(),
-				TArray<FString>(),
+				MoveTemp(PartNames),
 				HostOverride,
 				RequestFlags);
 		FString BuildTransferName = FString::Printf(TEXT("%s%s"), *Group.DisplayName, DownloadSpecJSONContents.IsEmpty() ? TEXT("") : TEXT("*"));
@@ -2307,12 +2307,12 @@ TSharedRef<SWidget> SBuildSelection::GetGridPanel()
 			[
 				SAssignNew(BuildArtifactSelection, SBuildArtifactSelection)
 					.BuildServiceInstance(BuildServiceInstance)
-					.OnDownloadWithSpec_Lambda([this](const FString& ArtifactName, const UE::BuildSelection::Internal::FArtifact& Artifact, const FString& Spec)
+					.OnDownloadWithSpec_Lambda([this](const FString& ArtifactName, const UE::BuildSelection::Internal::FArtifact& Artifact, const FString& Spec, const TArray<FString>& PartNames)
 					{
 						TArray<FBuildSelectionBuildGroupPtr> SelectedItems = BuildGroupListView->GetSelectedItems();
 						if (SelectedItems.Num() > 0)
 						{
-							DownloadItem(*SelectedItems[0], ArtifactName, Artifact, FString(Spec));
+							DownloadItem(*SelectedItems[0], ArtifactName, Artifact, FString(Spec), TArray<FString>(PartNames));
 						}
 					})
 			]
diff --git a/Engine/Source/Programs/BuildStorageTool/Private/BuildStorageToolStyle.cpp b/Engine/Source/Programs/BuildStorageTool/Private/BuildStorageToolStyle.cpp
--- a/Engine/Source/Programs/BuildStorageTool/Private/BuildStorageToolStyle.cpp
+++ b/Engine/Source/Programs/BuildStorageTool/Private/BuildStorageToolStyle.cpp
@@ -95,6 +95,30 @@ TSharedRef< FSlateStyleSet > FBuildStorageToolStyle::Create()
 
 	Style.Set("Artifact.SmallContext", SmallContextCombo);
 
+	// FlatButton styles — square-cornered buttons matching the editor's StarshipStyle definitions.
+	// Not inherited from CoreStyle so must be defined here for the standalone tool.
+	{
+		const FButtonStyle& BaseButton = FCoreStyle::Get().GetWidgetStyle<FButtonStyle>("Button");
+		const FLinearColor SelectionColor = FStyleColors::Primary.GetSpecifiedColor();
+		const FLinearColor SelectionColorPressed = SelectionColor * 0.75f;
+
+		Style.Set("FlatButton", FButtonStyle(BaseButton)
+			.SetNormal(FSlateNoResource())
+			.SetHovered(BOX_BRUSH("Common/FlatButton", 2.0f / 8.0f, SelectionColor))
+			.SetPressed(BOX_BRUSH("Common/FlatButton", 2.0f / 8.0f, SelectionColorPressed))
+		);
+
+		const FLinearColor DarkNormal(0.125f, 0.125f, 0.125f, 0.8f);
+
+		Style.Set("FlatButton.Dark", FButtonStyle(BaseButton)
+			.SetNormal(BOX_BRUSH("Common/FlatButton", 2.0f / 8.0f, DarkNormal))
+			.SetHovered(BOX_BRUSH("Common/FlatButton", 2.0f / 8.0f, SelectionColor))
+			.SetPressed(BOX_BRUSH("Common/FlatButton", 2.0f / 8.0f, SelectionColorPressed))
+		);
+
+		Style.Set("FlatButton.Default", Style.GetWidgetStyle<FButtonStyle>("FlatButton.Dark"));
+	}
+
 	const FTextBlockStyle DefaultText = FTextBlockStyle()
 		.SetFont(DEFAULT_FONT("Bold", 10))
 		.SetColorAndOpacity(FSlateColor::UseForeground())
