BuildStorageTool: Fix double scheme in download URL.

#Synced-CL 52440464

diff --git a/Engine/Source/Editor/StorageServerWidgets/Private/SBuildArtifactSelection.cpp b/Engine/Source/Editor/StorageServerWidgets/Private/SBuildArtifactSelection.cpp
--- a/Engine/Source/Editor/StorageServerWidgets/Private/SBuildArtifactSelection.cpp
+++ b/Engine/Source/Editor/StorageServerWidgets/Private/SBuildArtifactSelection.cpp
@@ -525,7 +525,7 @@ TSharedRef<SWidget> SBuildArtifactSelection::CreateArtifactMenuWidget(UE::BuildS
 						if (ServiceInstance)
 						{
 							// TODO: This URL format is specific to UE Build Storage - we should genericize this
-							FString Url = FString::Printf(TEXT("https://%s/api/v2/builds/%s/%s/%s"),
+							FString Url = FString::Printf(TEXT("%s/api/v2/builds/%s/%s/%s"),
 								*FString(ServiceInstance->GetEffectiveDomain()),
 								*Namespace,
 								*Artifact->BuildRecord.BucketId,
